Add an initial implementation of nonlinearly weighted first-order regression (NFOR) based offline denoiser (experimental). It helps to improve path tracing temporal stability for slow moving cameras when adjacent frame history is close.
## How to turn on and use the denoiser
1. Enable "NFOR Denoise" plugin.
2. Set console variable to "r.PathTracing.SpatialDenoiser.Type 1" to use the denoiser.
## MRQ setting
1. Set `r.NFOR.FrameCount 2` to use 5 frame history(2 previous frames, current frame, and 2 future frames). The value is default to 0 so even in the editor, we can preview single frame denoising. 0, 1, 2, 3 are allowed for the setting at this moment.
2. Turn on `reference motion blur` and use temporal samples.
## How to tune parameters in editor
1. Change NFOR parameters (e.g., r.NFOR.NonLocalMean.Radiance.PatchDistance 12).
2. Disable denoiser `r.PathTracing.Denoiser 0`.
3. Enable denoiser `r.PathTracing.Denoiser 1`. Now the new setting will be applied.
## Misc NFOR parameters
1. Radiance filtering parameters:
 *  r.NFOR.NonLocalMean.Radiance.PatchDistance (default: 9)
 *  r.NFOR.NonLocalMean.Radiance.PatchSize (default: 3)
2. Feature filtering parameters:
 * r.NFOR.NonLocalMean.Feature.PatchDistance (default: 5)
 * r.NFOR.NonLocalMean.Feature.PatchSize (default: 3)

diff --git a/Engine/Plugins/Experimental/NFORDenoise/NFORDenoise.uplugin b/Engine/Plugins/Experimental/NFORDenoise/NFORDenoise.uplugin
new file mode 100644
--- /dev/null
+++ b/Engine/Plugins/Experimental/NFORDenoise/NFORDenoise.uplugin
@@ -0,0 +1,33 @@
+{
+	"FileVersion": 3,
+	"Version": 1,
+	"VersionName": "1.0",
+	"FriendlyName": "NFORDenoise",
+	"Description": "Denoising engine for the Unreal Path Tracer based on Nonlinearly Weighted First-order Regression for Denoising Monte Carlo Renderings.",
+	"Category": "Denoising",
+	"CreatedBy": "Epic Games, Inc.",
+	"CreatedByURL": "https://epicgames.com",
+	"DocsURL": "",
+	"MarketplaceURL": "",
+	"SupportURL": "",
+	"EnabledByDefault": false,
+	"CanContainContent": false,
+	"IsBetaVersion": false,
+	"IsExperimentalVersion": true,
+	"Installed": false,
+	"SupportedTargetPlatforms": [
+		"Win64",
+		"Linux"
+	],
+	"Modules": [
+		{
+			"Name": "NFORDenoise",
+			"Type": "RuntimeAndProgram",
+			"LoadingPhase": "PostConfigInit",
+			"PlatformAllowList": [
+				"Win64",
+				"Linux"
+			]
+		}
+	]
+}
diff --git a/Engine/Plugins/Experimental/NFORDenoise/Source/NFORDenoise/NFORDenoise.Build.cs b/Engine/Plugins/Experimental/NFORDenoise/Source/NFORDenoise/NFORDenoise.Build.cs
new file mode 100644
--- /dev/null
+++ b/Engine/Plugins/Experimental/NFORDenoise/Source/NFORDenoise/NFORDenoise.Build.cs
@@ -0,0 +1,30 @@
+// Copyright Epic Games, Inc. All Rights Reserved.
+
+namespace UnrealBuildTool.Rules
+{
+	public class NFORDenoise : ModuleRules
+	{
+		public NFORDenoise(ReadOnlyTargetRules Target) : base(Target)
+		{
+
+			PublicDependencyModuleNames.AddRange(
+				new string[]
+				{
+							"Core",
+					// ... add other public dependencies that you statically link with here ...
+				}
+			);
+
+			PrivateDependencyModuleNames.AddRange(
+				new string[] {
+					"RenderCore",
+					"Renderer",
+					"RHI",
+					"Projects",	// Use shader folder remapping
+					"Eigen",	// Debug first order regression
+				}
+			);
+
+		}
+	}
+}
diff --git a/Engine/Plugins/Experimental/NFORDenoise/Source/NFORDenoise/Private/NFORDenoise.cpp b/Engine/Plugins/Experimental/NFORDenoise/Source/NFORDenoise/Private/NFORDenoise.cpp
new file mode 100644
--- /dev/null
+++ b/Engine/Plugins/Experimental/NFORDenoise/Source/NFORDenoise/Private/NFORDenoise.cpp
@@ -0,0 +1,275 @@
+// Copyright Epic Games, Inc. All Rights Reserved.
+
+#include "NFORDenoise.h"
+#include "Interfaces/IPluginManager.h"
+#include "PathTracingDenoiser.h"
+#include "RenderGraphBuilder.h"
+#include "RenderGraphUtils.h"
+#include "SceneView.h"
+#include "NFORDenoiseCS.h"
+
+#define LOCTEXT_NAMESPACE "FNFORDenoiseModule"
+
+DEFINE_LOG_CATEGORY(LogNFORDenoise)
+
+using namespace UE::Renderer::Private;
+
+static TAutoConsoleVariable<bool> CVarNFORFeatureDepth(
+	TEXT("r.NFOR.Feature.Depth"),
+	false,
+	TEXT("true: Add depth as feature for NFOR denoising."),
+	ECVF_RenderThreadSafe
+);
+
+bool IsFeatureDepthEnabled()
+{
+	return CVarNFORFeatureDepth.GetValueOnRenderThread();
+}
+
+struct FExternalFrameCache;
+
+struct FRDGFrameCache
+{
+	FRDGTextureRef Radiance;
+	FRDGTextureRef RadianceVariance;
+	FRDGTextureRef Albedo;
+	FRDGTextureRef Normal;
+
+	FExternalFrameCache ToExternalResource(FRDGBuilder& GraphBuilder, const FSceneView& View);
+	void AppendToNFORBuffer(TArray<NFORDenoise::FRadianceDesc>& Radiances, TArray<NFORDenoise::FFeatureDesc>& Features);
+};
+
+struct FExternalFrameCache
+{
+	TRefCountPtr<IPooledRenderTarget> Radiance;
+	TRefCountPtr<IPooledRenderTarget> RadianceVariance;
+	TRefCountPtr<IPooledRenderTarget> Albedo;
+	TRefCountPtr<IPooledRenderTarget> Normal;
+
+	FRDGFrameCache ToRGDResource(FRDGBuilder& GraphBuilder, const FSceneView& View);
+};
+
+FExternalFrameCache FRDGFrameCache::ToExternalResource(FRDGBuilder& GraphBuilder, const FSceneView& View)
+{
+	FExternalFrameCache ExternalFrameCache;
+	ExternalFrameCache.Radiance = GraphBuilder.ConvertToExternalTexture(Radiance);
+	ExternalFrameCache.RadianceVariance = GraphBuilder.ConvertToExternalTexture(RadianceVariance);
+	ExternalFrameCache.Albedo = GraphBuilder.ConvertToExternalTexture(Albedo);
+	ExternalFrameCache.Normal = GraphBuilder.ConvertToExternalTexture(Normal);
+
+	return ExternalFrameCache;
+}
+
+void FRDGFrameCache::AppendToNFORBuffer(
+	TArray<NFORDenoise::FRadianceDesc>& Radiances,
+	TArray<NFORDenoise::FFeatureDesc>& Features)
+{
+	{
+		NFORDenoise::FNFORTextureDesc RadianceTex(Radiance, 0, 3, 4);
+		NFORDenoise::FNFORTextureDesc RadianceVarianceTex(RadianceVariance, 0, 1, 4);
+		Radiances.Add(NFORDenoise::FRadianceDesc(RadianceTex, RadianceVarianceTex, NFORDenoise::EVarianceType::GreyScale));
+	}
+
+	{
+		NFORDenoise::FNFORTextureDesc AlbedoTex(Albedo, 0, 3, 4);
+		NFORDenoise::FNFORTextureDesc AlbedoVarianceTex(RadianceVariance, 1, 1, 4);
+
+		NFORDenoise::FNFORTextureDesc NormalTex(Normal, 0, 3, 4);
+		NFORDenoise::FNFORTextureDesc NormalVarianceTex(RadianceVariance, 2, 1, 4);
+
+		Features.Add(NFORDenoise::FFeatureDesc(AlbedoTex, AlbedoVarianceTex, NFORDenoise::EVarianceType::GreyScale));
+		Features.Add(NFORDenoise::FFeatureDesc(NormalTex, NormalVarianceTex, NFORDenoise::EVarianceType::GreyScale));
+
+		if (IsFeatureDepthEnabled())
+		{
+			NFORDenoise::FNFORTextureDesc DepthTex(Normal, 3, 1, 4);
+			NFORDenoise::FNFORTextureDesc DepthVarianceTex(nullptr);
+			Features.Add(NFORDenoise::FFeatureDesc(DepthTex, DepthVarianceTex, NFORDenoise::EVarianceType::GreyScale, true));
+		}
+	}
+}
+
+FRDGFrameCache FExternalFrameCache::ToRGDResource(FRDGBuilder& GraphBuilder, const FSceneView& View)
+{
+	FRDGFrameCache RDGFrameCache;
+	RDGFrameCache.Radiance = GraphBuilder.RegisterExternalTexture(Radiance);
+	RDGFrameCache.RadianceVariance = GraphBuilder.RegisterExternalTexture(RadianceVariance);
+	RDGFrameCache.Albedo = GraphBuilder.RegisterExternalTexture(Albedo);
+	RDGFrameCache.Normal = GraphBuilder.RegisterExternalTexture(Normal);
+
+	return RDGFrameCache;
+}
+
+
+class FNFORDenosier : public IPathTracingSpatialTemporalDenoiser
+{
+public:
+	class FHistory : public IHistory
+	{
+	public:
+		FHistory(const TCHAR* DebugName) : DebugName(DebugName) {}
+
+		virtual ~FHistory() 
+		{ 
+		}
+
+		const TCHAR* GetDebugName() const override { return DebugName; };
+
+
+		void AddFrame(FRDGBuilder& GraphBuilder, const FSceneView& View, const FInputs& Inputs)
+		{
+			const int NumFrames = NFORDenoise::GetFrameCount();
+			if (NumFrames <= 0)
+			{
+				return;
+			}
+
+			FRDGFrameCache RDGFrameCache;
+			RDGFrameCache.Radiance = GraphBuilder.CreateTexture(Inputs.ColorTex->Desc,TEXT("Radiance"));
+			RDGFrameCache.RadianceVariance = GraphBuilder.CreateTexture(Inputs.VarianceTex->Desc, TEXT("RadianceVariance"));
+			RDGFrameCache.Albedo = GraphBuilder.CreateTexture(Inputs.AlbedoTex->Desc, TEXT("Albedo"));
+			RDGFrameCache.Normal = GraphBuilder.CreateTexture(Inputs.NormalTex->Desc, TEXT("Normal"));
+
+			AddCopyTexturePass(GraphBuilder, Inputs.ColorTex, RDGFrameCache.Radiance);
+			AddCopyTexturePass(GraphBuilder, Inputs.VarianceTex, RDGFrameCache.RadianceVariance);
+			AddCopyTexturePass(GraphBuilder, Inputs.AlbedoTex, RDGFrameCache.Albedo);
+			AddCopyTexturePass(GraphBuilder, Inputs.NormalTex, RDGFrameCache.Normal);
+			
+			FExternalFrameCache FrameCache = RDGFrameCache.ToExternalResource(GraphBuilder, View);
+
+			while (FrameCaches.Num() >= NumFrames)
+			{
+				FrameCaches.RemoveAt(0);
+			}
+			FrameCaches.Add(MoveTemp(FrameCache));
+		}
+
+		void FetchFrames(FRDGBuilder& GraphBuilder, const FSceneView& View, 
+			TArray<NFORDenoise::FRadianceDesc>& Radiances,
+			TArray<NFORDenoise::FFeatureDesc>& FeatureDescs,
+			bool bReverseOrder = true)
+		{
+			const int TotalCachedFrames = FrameCaches.Num();
+			for (int i = 0; i < TotalCachedFrames; ++i)
+			{
+				int Index = bReverseOrder ? (TotalCachedFrames - i - 1) : i;
+				FExternalFrameCache& FrameCache = FrameCaches[Index];
+				FRDGFrameCache RDGFrameCache = FrameCache.ToRGDResource(GraphBuilder, View);
+				RDGFrameCache.AppendToNFORBuffer(Radiances, FeatureDescs);
+			}
+		}
+
+		FRDGTextureRef GetDepth(FRDGBuilder& GraphBuilder, int32 FrameIndex)
+		{
+			check(FrameIndex >=0 && FrameIndex < Num());
+			TRefCountPtr<IPooledRenderTarget> NormalDepth= FrameCaches[FrameIndex].Normal;
+			return NormalDepth ? GraphBuilder.RegisterExternalTexture(NormalDepth) : nullptr;
+		}
+
+		int32 Num() const { return FrameCaches.Num(); }
+
+	private:
+		const TCHAR* DebugName;
+
+		TArray<FExternalFrameCache> FrameCaches;
+	};
+
+	~FNFORDenosier() {}
+
+	const TCHAR* GetDebugName() const override { return *DebugName; }
+
+	FOutputs AddPasses(FRDGBuilder& GraphBuilder, const FSceneView& View, const FInputs& Inputs) const override
+	{
+		AddCopyTexturePass(GraphBuilder,Inputs.ColorTex, Inputs.OutputTex);
+
+		FHistory* History = Inputs.PrevHistory.IsValid() ? reinterpret_cast<FHistory*>(Inputs.PrevHistory.GetReference()) : nullptr;
+
+		if (!Inputs.VarianceTex)
+		{
+			return { History };
+		}
+		
+		TUniquePtr<FHistory> CurHistory = MakeUnique<FHistory>(*DebugName);
+		if (History)
+		{
+			CurHistory.Reset(History);
+		}
+
+		CurHistory->AddFrame(GraphBuilder, View, Inputs);
+		
+		// Fetch all the frames for denoising
+		TArray<NFORDenoise::FFeatureDesc> Features;
+		TArray<NFORDenoise::FRadianceDesc> Radiances;
+
+		CurHistory->FetchFrames(GraphBuilder, View, Radiances, Features);
+
+		// Denoising based on spatial temporal denoising config
+		const bool bDenoised = NFORDenoise::FilterMain(
+				GraphBuilder,
+				View,
+				Radiances,
+				Features, 
+				Inputs.OutputTex);
+
+		UE_LOG(LogNFORDenoise, Log, TEXT("Frame %d: %s NumOfHistory = %d"), Inputs.DenoisingFrameId, bDenoised ? TEXT("Denoised") : TEXT("Accumulating history"), CurHistory ? CurHistory->Num() : 0);
+
+		int32 NumberOfHistory = CurHistory->Num();
+
+		// Post process DOF replies on the depth buffer to estimate the DOF,
+		// Update the output depth buffer to the depth of the center image if the denoising frame is not the latest.
+		const bool bNeedToUpdateDepth = bDenoised && (NFORDenoise::GetFrameCount() > 1) && NFORDenoise::GetDenoisingFrameIndex(NumberOfHistory) > 0;
+		if (bNeedToUpdateDepth)
+		{
+				
+			int32 DenoisingDepthFrameIndex = NumberOfHistory - NFORDenoise::GetDenoisingFrameIndex(NumberOfHistory) - 1;
+			if (DenoisingDepthFrameIndex >= 0)
+			{
+				const bool bAlphaOnly = true;
+				FRDGTextureRef DenoisingDepth = CurHistory->GetDepth(GraphBuilder, DenoisingDepthFrameIndex);
+
+				NFORDenoise::AddCopyMirroredTexturePass(GraphBuilder, DenoisingDepth, Inputs.NormalTex,
+					FIntPoint::ZeroValue, FIntPoint::ZeroValue, FIntPoint::ZeroValue, bAlphaOnly);
+			}
+		}
+
+		// TODO: Output denoised albedo, normal if needed.
+		return { CurHistory.Release() };
+	}
+
+	void AddMotionVectorPass(FRDGBuilder& GraphBuilder, const FSceneView& View, const FMotionVectorInputs& Inputs) const override
+	{
+		AddClearRenderTargetPass(GraphBuilder, Inputs.OutputTex, FLinearColor::Black);
+	}
+
+private:
+	inline static const FString DebugName = TEXT("FNFORDenosier");
+};
+
+void FNFORDenoiseModule::StartupModule()
+{
+	UE_LOG(LogNFORDenoise, Log, TEXT("NeuralPostProcessing starting up"));
+
+	const TSharedPtr<IPlugin> Plugin = IPluginManager::Get().FindPlugin(TEXT("NFORDenoise"));
+	if (Plugin.IsValid())
+	{
+		FString ModuleDir = Plugin->GetBaseDir() + TEXT("/Source/NFORDenoise");
+		AddShaderSourceDirectoryMapping(TEXT("/NFORDenoise"), FPaths::Combine(ModuleDir, TEXT("Shaders")));
+
+		GPathTracingSpatialTemporalDenoiserPlugin = MakeUnique<FNFORDenosier>();
+	}
+	else
+	{
+		UE_LOG(LogNFORDenoise, Error, TEXT("Shaders directory not added. Failed to find NFORDenoise plugin"));
+	}
+}
+
+void FNFORDenoiseModule::ShutdownModule()
+{
+	UE_LOG(LogNFORDenoise, Log, TEXT("NFORDenoise function shutting down"));
+
+	GPathTracingSpatialTemporalDenoiserPlugin.Reset();
+
+}
+
+#undef LOCTEXT_NAMESPACE
+IMPLEMENT_MODULE(FNFORDenoiseModule, NFORDenoise)
\ No newline at end of file
diff --git a/Engine/Plugins/Experimental/NFORDenoise/Source/NFORDenoise/Private/NFORDenoiseCS.cpp b/Engine/Plugins/Experimental/NFORDenoise/Source/NFORDenoise/Private/NFORDenoiseCS.cpp
new file mode 100644
--- /dev/null
+++ b/Engine/Plugins/Experimental/NFORDenoise/Source/NFORDenoise/Private/NFORDenoiseCS.cpp
@@ -0,0 +1,1748 @@
+// Copyright Epic Games, Inc. All Rights Reserved.
+
+#include "NFORDenoiseCS.h"
+#include "SystemTextures.h"
+#include "PixelShaderUtils.h"
+
+#include "NFORWeightedLSRCommon.h"
+#include "NFORRegressionCPUSolver.h"
+
+namespace NFORDenoise
+{
+	
+	TAutoConsoleVariable<bool> CVarNFORFeatureAddConstant(
+		TEXT("r.NFOR.Feature.AddConstant"),
+		1,
+		TEXT("Add a constant 1 feature for denoising. Especially useful when all other features are zero."),
+		ECVF_RenderThreadSafe);
+
+	TAutoConsoleVariable<int32> CVarNFORPredivideAlbedo(
+		TEXT("r.NFOR.PredivideAlbedo"),
+		1,
+		TEXT("Enable pre-albedo divide to denoise only the demodulated singal. It preserves more high frequency details."),
+		ECVF_RenderThreadSafe);
+
+	TAutoConsoleVariable<float> CVarNFORPredivideAlbedoOffset(
+		TEXT("r.NFOR.PredivideAlbedo.Offset"),
+		1e-3,
+		TEXT("Offset for albedo for regions other than full reflection and sky materials. Increase to get a smoother result."),
+		ECVF_RenderThreadSafe);
+
+	TAutoConsoleVariable<float> CVarNFORPredivideAlbedoOffsetSky(
+		TEXT("r.NFOR.PredivideAlbedo.OffsetSky"),
+		0.2,
+		TEXT("Sky or reflection of sky material has very small albedo that will cause noise. Offset more."),
+		ECVF_RenderThreadSafe);
+
+	TAutoConsoleVariable<int32> CVarNFORFrameCount(
+		TEXT("r.NFOR.FrameCount"),
+		0,
+		TEXT("n: Use the previous n frames, the current frame, and the future n frames. Suggested range is 0~2. Max=3."),
+		ECVF_RenderThreadSafe);
+
+	TAutoConsoleVariable<int32> CVarNFORFrameCountCondition(
+		TEXT("r.NFOR.FrameCount.Condition"),
+		1,
+		TEXT("0: Denoise even if the frame count accumulated is less than the required frame count (used for debug).")
+		TEXT("1: Denoise only when the number of frame count meets requirement."),
+		ECVF_RenderThreadSafe);
+
+	TAutoConsoleVariable<int32> CVarNFORDenoisingFrameIndex(
+		TEXT("r.NFOR.DenoisingFrameIndex"),
+		-1,
+		TEXT("The index of the denoising frame.")
+		TEXT("-1: Automatically determine the index.")
+		TEXT("i: Use all frames other than the ith frame to denoise."),
+		ECVF_RenderThreadSafe);
+
+	TAutoConsoleVariable<int32> CVarNFORNumOfTile(
+		TEXT("r.NFOR.NumOfTile"),
+		10,
+		TEXT("<=1: Use a single dispatch. Could run out of memory.\n")
+		TEXT("n: Divide the image into n x n tiles in [1,32].\n"),
+		ECVF_RenderThreadSafe);
+
+	TAutoConsoleVariable<int32> CVarNFORTileDebug(
+		TEXT("r.NFOR.TileDebug"),
+		0,
+		TEXT(">0: Turn on tile debug mode."),
+		ECVF_RenderThreadSafe);
+
+	TAutoConsoleVariable<int32> CVarNFORTileDebugIndex(
+		TEXT("r.NFOR.TileDebug.Index"),
+		-1,
+		TEXT("Tile index number to debug.")
+		TEXT(" -1: The middle index in range of 0 ~ (NumOfTile * NumOfTile - 1).")
+		TEXT(">=0: Select a specific tile to render for debug."),
+		ECVF_RenderThreadSafe);
+
+	TAutoConsoleVariable<int32> CVarNFORRegressionDevice(
+		TEXT("r.NFOR.Regression.Device"),
+		1,
+		TEXT(" 0: CPU (verification). Used only for feature development.\n")
+		TEXT(" 1: GPU.\n"),
+		ECVF_RenderThreadSafe);
+
+	TAutoConsoleVariable<float> CVarNFORRegressionDataRatioToParameters(
+		TEXT("r.NFOR.Regression.MaxDataRatioToParemters"),
+		20.0f,
+		TEXT("The max number of observations per parameter in the regression. <1 to use all."),
+		ECVF_RenderThreadSafe);
+
+	TAutoConsoleVariable<int32> CVarNFORLinearSolverDevice(
+		TEXT("r.NFOR.LinearSolver.Device"),
+		1,
+		TEXT("0: Solve Ax=B on CPU. Use householder QR decomposition from Eigen library.")
+		TEXT("1: Solve Ax=B on GPU."),
+		ECVF_RenderThreadSafe);
+
+	TAutoConsoleVariable<int32> CVarNFORReconstructionType(
+		TEXT("r.NFOR.Reconstruction.Type"),
+		0,
+		TEXT("0: Scatter for the denoising frame, gather for other temporal frames (default).")
+		TEXT("1: Force gathering."),
+		ECVF_RenderThreadSafe);
+
+	TAutoConsoleVariable<int32> CVarNFORReconstructionDebugFrameIndex(
+		TEXT("r.NFOR.Reconstruction.Debug.FrameIndex"),
+		-1,
+		TEXT(">=0: Output the denoising contribution from the ith frame only.")
+		TEXT("-1: do not perform debug. Output contributions from all frames."),
+		ECVF_RenderThreadSafe);
+
+	TAutoConsoleVariable<int32> CVarNFORNonLocalMeanFeaturePatchSize(
+		TEXT("r.NFOR.NonLocalMean.Feature.PatchSize"),
+		3,
+		TEXT("The patch size of the non-local mean algorithm for feature filtering. The patch width/height = PatchSize * 2 + 1."),
+		ECVF_RenderThreadSafe);
+
+	TAutoConsoleVariable<int32> CVarNFORNonLocalMeanFeaturePatchDistance(
+		TEXT("r.NFOR.NonLocalMean.Feature.PatchDistance"),
+		5,
+		TEXT("The search distance of the non-local mean algorithm for feature filtering. The searching patch width/height = PatchDistance * 2 + 1."),
+		ECVF_RenderThreadSafe);
+
+	TAutoConsoleVariable<int32> CVarNFORNonLocalMeanRadiancePatchSize(
+		TEXT("r.NFOR.NonLocalMean.Radiance.PatchSize"),
+		3,
+		TEXT("The patch size of the non-local mean algorithm. The patch width/height = PatchSize * 2 + 1."),
+		ECVF_RenderThreadSafe);
+
+	TAutoConsoleVariable<int32> CVarNFORNonLocalMeanRadiancePatchDistance(
+		TEXT("r.NFOR.NonLocalMean.Radiance.PatchDistance"),
+		9,
+		TEXT("The search distance of the non-local mean algorithm. The searching patch width/height = PatchDistance * 2 + 1.")
+		TEXT("The patch distance for bandwidth selection dependents on this parameters for MSE and selection filtering."),
+		ECVF_RenderThreadSafe);
+
+	TAutoConsoleVariable<bool> CVarNFORBandwidthSelection(
+		TEXT("r.NFOR.BandwidthSelection"),
+		true,
+		TEXT("true: Apply bandwidth selection. It helps to preserve both high and low frequency details."),
+		ECVF_RenderThreadSafe);
+
+	TAutoConsoleVariable<float> CVarNFORBandwidthSelectionBandwidth(
+		TEXT("r.NFOR.BandwidthSelection.Bandwidth"),
+		-1,
+		TEXT("-1: Use predefined bandwidths {0.5f, 1.0f}.")
+		TEXT("(0,1]: Use a specific bandwidth."),
+		ECVF_RenderThreadSafe);
+
+	// Working in progress
+	TAutoConsoleVariable<int32> CVarNFORAlbedoDivideRecoverPhase(
+		TEXT("r.NFOR.AlbedoDivide.RecoverPhase"),
+		0,
+		TEXT("0: Add back in the last step. Denoised = Albedo_{center} * \\sum_{all frames}{denoised radiance}. Require high sample count for high quality albedo.")
+		TEXT("1: Add back at each scattering or gathering. Denoised = * \\sum_{i \\in frames}{Albedo_i * denoised radiance}."),
+		ECVF_RenderThreadSafe
+	);
+
+	//--------------------------------------------------------------------------------------------------------------------
+	// Functions based on CVars
+
+	bool ShouldFeatureAddConstant()
+	{
+		return CVarNFORFeatureAddConstant.GetValueOnRenderThread();
+	}
+
+	bool IsPreAlbedoDivideEnabled()
+	{
+		const bool bPredivideAlbedo = CVarNFORPredivideAlbedo.GetValueOnRenderThread() != 0;
+		return bPredivideAlbedo;
+	}
+
+	FLinearColor GetPreAlbedoDivideAlbedoOffset()
+	{
+		const float Offset = FMath::Max(1e-8, CVarNFORPredivideAlbedoOffset.GetValueOnRenderThread());
+		const float OffsetSky = FMath::Max(1e-8, CVarNFORPredivideAlbedoOffsetSky.GetValueOnRenderThread());
+		return FLinearColor(Offset, OffsetSky, 0.0f);
+	}
+
+	int32 GetFrameCount()
+	{
+		const int32 NumFrames = FMath::Clamp(1 + 2 * CVarNFORFrameCount.GetValueOnRenderThread(), 1, 7);
+		return NumFrames;
+	}
+
+	enum class EDenoiseFrameCountCondition : uint32
+	{
+		Any,
+		Equal,
+		MAX
+	};
+
+	EDenoiseFrameCountCondition GetFrameCountCoundition()
+	{
+		uint32 ConditionValue = CVarNFORFrameCountCondition.GetValueOnRenderThread();
+		EDenoiseFrameCountCondition Condition = EDenoiseFrameCountCondition::Any;
+		if (ConditionValue != 0)
+		{
+			Condition = EDenoiseFrameCountCondition::Equal;
+		}
+		return Condition;
+	}
+
+	int32 GetDenoisingFrameIndex(int32 NumberOfFrameInBuffer)
+	{
+		int32 TargetFrameCount = GetFrameCount();
+		int32 DenoisingFrameIndex = CVarNFORDenoisingFrameIndex.GetValueOnRenderThread();
+		int32 ResolvedSourceFrameIndex = INDEX_NONE;
+		if (DenoisingFrameIndex < 0)
+		{
+			// If no specific denoising frame index is specified, use the center one
+			if (NumberOfFrameInBuffer >= 0)
+			{
+				ResolvedSourceFrameIndex = (NumberOfFrameInBuffer > TargetFrameCount / 2) ? (TargetFrameCount / 2) : INDEX_NONE;
+			}
+			else
+			{
+				ResolvedSourceFrameIndex = TargetFrameCount / 2;
+			}
+		}
+		else
+		{
+			// If the user has set the denoising index, use the available index within the limit
+			ResolvedSourceFrameIndex = FMath::Clamp(DenoisingFrameIndex, 0, TargetFrameCount - 1);
+			if (NumberOfFrameInBuffer - 1 < ResolvedSourceFrameIndex)
+			{
+				ResolvedSourceFrameIndex = INDEX_NONE;
+			}
+		}
+		
+		return ResolvedSourceFrameIndex;
+	}
+
+	int32 GetNumOfTiles()
+	{
+		return FMath::Clamp(CVarNFORNumOfTile.GetValueOnRenderThread(), 1, 32);
+	}
+
+	bool IsTileDebugEnabled()
+	{
+		return CVarNFORTileDebug.GetValueOnRenderThread() > 0;
+	}
+
+	int32 GetTileDebugIndex()
+	{
+		return CVarNFORTileDebugIndex.GetValueOnRenderThread();
+	}
+
+	enum class ERegressionDevice : int32
+	{
+		CPU,
+		GPU,
+		MAX
+	};
+
+	ERegressionDevice GetRegressionDevice()
+	{
+		const int32 RegressionDevice = FMath::Clamp(CVarNFORRegressionDevice.GetValueOnRenderThread(),
+			static_cast<int32>(ERegressionDevice::CPU),
+			static_cast<int32>(ERegressionDevice::MAX) - 1);
+		return static_cast<ERegressionDevice>(RegressionDevice);
+	}
+
+	int32 GetSamplingStep(int32 NumberOfParameters, int32 TotalDataRecords)
+	{
+		int32 DataRatioToParameters = CVarNFORRegressionDataRatioToParameters.GetValueOnRenderThread();
+		if (DataRatioToParameters < 1)
+		{
+			return 1;
+		}
+		return FMath::Max(1, TotalDataRecords / (NumberOfParameters * DataRatioToParameters));
+	}
+
+	enum class ELinearSolverDevice : int32
+	{
+		CPU,
+		GPU,
+		MAX
+	};
+
+	ELinearSolverDevice GetLinearSolverDevice()
+	{
+		const int32 LinearSolverDevice = FMath::Clamp(CVarNFORLinearSolverDevice.GetValueOnRenderThread(),
+			static_cast<int32>(ELinearSolverDevice::CPU),
+			static_cast<int32>(ELinearSolverDevice::MAX) - 1);
+		return static_cast<ELinearSolverDevice>(LinearSolverDevice);
+	}
+
+	RegressionKernel::FReconstructSpatialTemporalImage::EReconstructionType 
+		GetReconstructionType( int32 CurrentFrameIndex, int32 DenoisingFrameIndex)
+	{
+		RegressionKernel::FReconstructSpatialTemporalImage::EReconstructionType ReconstructionType;
+		ReconstructionType = RegressionKernel::FReconstructSpatialTemporalImage::EReconstructionType::Gather;
+
+		if (CurrentFrameIndex == DenoisingFrameIndex)
+		{
+			ReconstructionType = RegressionKernel::FReconstructSpatialTemporalImage::EReconstructionType::Scatter;
+		}
+
+		if (CVarNFORReconstructionType.GetValueOnRenderThread() != 0)
+		{
+			ReconstructionType = RegressionKernel::FReconstructSpatialTemporalImage::EReconstructionType::Gather;
+		}
+
+		return ReconstructionType;
+	}
+
+	int32 GetReconstructionDebugFrameIndex()
+	{
+		return CVarNFORReconstructionDebugFrameIndex.GetValueOnRenderThread();
+	}
+
+	int32 GetNonLocalMeanFeaturePatchSize()
+	{
+		return FMath::Clamp(CVarNFORNonLocalMeanFeaturePatchSize.GetValueOnRenderThread(), 0, 10);
+	}
+
+	int32 GetNonLocalMeanFeaturePatchDistance()
+	{
+		return FMath::Clamp(CVarNFORNonLocalMeanFeaturePatchDistance.GetValueOnRenderThread(), 0, 30);
+	}
+
+	int32 GetNonLocalMeanRadiancePatchSize()
+	{
+		return FMath::Clamp(CVarNFORNonLocalMeanRadiancePatchSize.GetValueOnRenderThread(), 0, 10);
+	}
+
+	int32 GetNonLocalMeanRadiancePatchDistance()
+	{
+		return FMath::Clamp(CVarNFORNonLocalMeanRadiancePatchDistance.GetValueOnRenderThread(), 0, 30);
+	}
+
+	bool IsBandwidthSelectionEnabled()
+	{
+		return CVarNFORBandwidthSelection.GetValueOnRenderThread();
+	}
+
+	TArray<float> GetBandwidthsConfiguration()
+	{
+		TArray<float> Bandwidths = { 0.5f, 1.0f };
+		{
+			float BandWidthOverride = FMath::Min(CVarNFORBandwidthSelectionBandwidth.GetValueOnRenderThread(), 1.0f);
+			if (BandWidthOverride > 0)
+			{
+				Bandwidths = { BandWidthOverride };
+			}
+		}
+
+		return Bandwidths;
+	}
+
+	EAlbedoDivideRecoverPhase GetPreAlbedoDivideRecoverPhase()
+	{
+		EAlbedoDivideRecoverPhase AlbedoDivideRecoverPhase = EAlbedoDivideRecoverPhase::Disabled;
+		if (IsPreAlbedoDivideEnabled())
+		{
+			if (CVarNFORAlbedoDivideRecoverPhase.GetValueOnRenderThread() == 0)
+			{
+				AlbedoDivideRecoverPhase = EAlbedoDivideRecoverPhase::Final;
+			}
+			else
+			{
+				AlbedoDivideRecoverPhase = EAlbedoDivideRecoverPhase::Each;
+			}
+		}
+		return AlbedoDivideRecoverPhase;
+	}
+
+	//--------------------------------------------------------------------------------------------------------------------
+	// Shader implementations
+	//--------------------------------------------------------------------------------------------------------------------
+	// General texture operations
+	IMPLEMENT_GLOBAL_SHADER(FTextureMultiplyCS, "/NFORDenoise/NFORDenoise.usf", "TextureOperationCS", SF_Compute);
+	IMPLEMENT_GLOBAL_SHADER(FTextureDivideCS, "/NFORDenoise/NFORDenoise.usf", "TextureOperationCS", SF_Compute);
+	IMPLEMENT_GLOBAL_SHADER(FTextureAccumulateConstantCS, "/NFORDenoise/NFORDenoise.usf", "TextureOperationCS", SF_Compute);
+	IMPLEMENT_GLOBAL_SHADER(FTextureAccumulateCS, "/NFORDenoise/NFORDenoise.usf", "TextureOperationCS", SF_Compute);
+	IMPLEMENT_GLOBAL_SHADER(FCopyTexturePS, "/NFORDenoise/NFORDenoise.usf", "CopyTexturePS", SF_Pixel);
+
+	//--------------------------------------------------------------------------------------------------------------------
+	// Radiance normalization
+	IMPLEMENT_GLOBAL_SHADER(FClassifyPreAlbedoDivideMaskIdCS, "/NFORDenoise/NFORDenoise.usf", "ClassifyPreAlbedoDivideMaskIdCS", SF_Compute);
+	IMPLEMENT_GLOBAL_SHADER(FNormalizeRadianceVarianceByAlbedoCS, "/NFORDenoise/NFORDenoise.usf", "NormalizeRadianceVarianceByAlbedoCS", SF_Compute);
+
+	//--------------------------------------------------------------------------------------------------------------------
+	// Non-local mean weight and filtering
+	IMPLEMENT_GLOBAL_SHADER(FNonLocalMeanFilteringCS, "/NFORDenoise/NFORDenoise.usf", "NonLocalMeanFilteringCS", SF_Compute);
+	IMPLEMENT_GLOBAL_SHADER(FNonLocalMeanWeightsCS, "/NFORDenoise/NFORDenoise.usf", "NonLocalMeanWeightsCS", SF_Compute);
+
+	//--------------------------------------------------------------------------------------------------------------------
+	// Collaborative filtering
+	//	1. Tiling
+	IMPLEMENT_GLOBAL_SHADER(FCopyTextureToBufferCS, "/NFORDenoise/NFORDenoise.usf", "CopyTextureToBufferCS", SF_Compute);
+	IMPLEMENT_GLOBAL_SHADER(FNormalizeTextureCS, "/NFORDenoise/NFORDenoise.usf", "NormalizeTextureCS", SF_Compute);
+	
+	//	2. Weighted Least-square solver
+	IMPLEMENT_GLOBAL_SHADER(RegressionKernel::FInPlaceBatchedMatrixMultiplicationCS, "/NFORDenoise/NFORDenoise.usf", "InPlaceBatchedMatrixMultiplicationCS", SF_Compute);
+	IMPLEMENT_GLOBAL_SHADER(RegressionKernel::FLinearSolverCS, "/NFORDenoise/NFORDenoise.usf", "LinearSolverCS", SF_Compute);
+	IMPLEMENT_GLOBAL_SHADER(RegressionKernel::FReconstructSpatialTemporalImage, "/NFORDenoise/NFORDenoise.usf", "ReconstructSpatialTemporalImageCS", SF_Compute);
+	IMPLEMENT_GLOBAL_SHADER(FAccumulateBufferToTextureCS, "/NFORDenoise/NFORDenoise.usf", "AccumulateBufferToTextureCS", SF_Compute);
+
+	//--------------------------------------------------------------------------------------------------------------------
+	// Bandwidth selection
+	IMPLEMENT_GLOBAL_SHADER(FMSEEstimationCS, "/NFORDenoise/NFORDenoise.usf", "MSEEstimationCS", SF_Compute);
+	IMPLEMENT_GLOBAL_SHADER(FGenerateSelectionMapCS, "/NFORDenoise/NFORDenoise.usf", "GenerateSelectionMapCS", SF_Compute);
+	IMPLEMENT_GLOBAL_SHADER(FCombineFilteredImageCS, "/NFORDenoise/NFORDenoise.usf", "CombineFilteredImageCS", SF_Compute);
+
+	//--------------------------------------------------------------------------------------------------------------------
+	// General texture operations
+	void AddMultiplyTextureRegionPass(FRDGBuilder& GraphBuilder, const FRDGTextureRef& SourceTexture, const FRDGTextureRef& TargetTexture,
+		bool bForceMultiply, FIntPoint SourcePosition, FIntPoint TargetPosition, FIntPoint Size)
+	{
+		Size = Size == FIntPoint::ZeroValue ? SourceTexture->Desc.Extent : Size;
+		FGlobalShaderMap* GlobalShaderMap = GetGlobalShaderMap(GMaxRHIFeatureLevel);
+
+		typedef FTextureMultiplyCS SHADER;
+		SHADER::FParameters* PassParameters = GraphBuilder.AllocParameters<SHADER::FParameters>();
+		{
+			PassParameters->Source = GraphBuilder.CreateSRV(FRDGTextureSRVDesc::Create(SourceTexture));
+			PassParameters->RWTarget = GraphBuilder.CreateUAV(TargetTexture);
+			PassParameters->SourcePosition = SourcePosition;
+			PassParameters->TargetPosition = TargetPosition;
+			PassParameters->ForceOperation = static_cast<int32>(bForceMultiply);
+			PassParameters->Size = Size;
+		}
+
+		TShaderMapRef<SHADER> ComputeShader(GlobalShaderMap);
+
+		FComputeShaderUtils::AddPass(
+			GraphBuilder,
+			RDG_EVENT_NAME("NFOR::AddMultiplyTextureRegionPass (%s [%d,%d] -> %s [%d,%d], size:%dx%d)",
+				SourceTexture->Name,
+				SourcePosition.X,
+				SourcePosition.Y,
+				TargetTexture->Name,
+				TargetPosition.X,
+				TargetPosition.Y,
+				Size.X,
+				Size.Y),
+			ERDGPassFlags::Compute,
+			ComputeShader,
+			PassParameters,
+			FComputeShaderUtils::GetGroupCount(Size, NON_LOCAL_MEAN_THREAD_GROUP_SIZE));
+	}
+
+	void AddDivideTextureRegionPass(FRDGBuilder& GraphBuilder, const FRDGTextureRef& SourceTexture, const FRDGTextureRef& TargetTexture,
+		bool bForceDivide, FIntPoint SourcePosition, FIntPoint TargetPosition, FIntPoint Size)
+	{
+		Size = Size == FIntPoint::ZeroValue ? SourceTexture->Desc.Extent : Size;
+		FGlobalShaderMap* GlobalShaderMap = GetGlobalShaderMap(GMaxRHIFeatureLevel);
+
+		typedef FTextureDivideCS SHADER;
+		SHADER::FParameters* PassParameters = GraphBuilder.AllocParameters<SHADER::FParameters>();
+		{
+			PassParameters->Source = GraphBuilder.CreateSRV(FRDGTextureSRVDesc::Create(SourceTexture));
+			PassParameters->RWTarget = GraphBuilder.CreateUAV(TargetTexture);
+			PassParameters->SourcePosition = SourcePosition;
+			PassParameters->TargetPosition = TargetPosition;
+			PassParameters->ForceOperation = static_cast<int32>(bForceDivide);
+			PassParameters->Size = Size;
+		}
+
+		TShaderMapRef<SHADER> ComputeShader(GlobalShaderMap);
+
+		FComputeShaderUtils::AddPass(
+			GraphBuilder,
+			RDG_EVENT_NAME("NFOR::AddDivideTextureRegionPass (%s [%d,%d] -> %s [%d,%d], size:%dx%d)",
+				SourceTexture->Name,
+				SourcePosition.X,
+				SourcePosition.Y,
+				TargetTexture->Name,
+				TargetPosition.X,
+				TargetPosition.Y,
+				Size.X,
+				Size.Y),
+			ERDGPassFlags::Compute,
+			ComputeShader,
+			PassParameters,
+			FComputeShaderUtils::GetGroupCount(Size, NON_LOCAL_MEAN_THREAD_GROUP_SIZE));
+	}
+
+	void AddAccumulateTextureRegionPass(FRDGBuilder& GraphBuilder, const FRDGTextureRef& SourceTexture, const FRDGTextureRef& TargetTexture,
+		FIntPoint SourcePosition, FIntPoint TargetPosition, FIntPoint Size)
+	{
+		Size = Size == FIntPoint::ZeroValue ? SourceTexture->Desc.Extent : Size;
+		FGlobalShaderMap* GlobalShaderMap = GetGlobalShaderMap(GMaxRHIFeatureLevel);
+
+		{
+			typedef FTextureAccumulateCS SHADER;
+			SHADER::FParameters* PassParameters = GraphBuilder.AllocParameters<SHADER::FParameters>();
+			{
+				PassParameters->Source = GraphBuilder.CreateSRV(FRDGTextureSRVDesc::Create(SourceTexture));
+				PassParameters->RWTarget = GraphBuilder.CreateUAV(TargetTexture);
+				PassParameters->SourcePosition = SourcePosition;
+				PassParameters->TargetPosition = TargetPosition;
+				PassParameters->Size = Size;
+			}
+
+			TShaderMapRef<SHADER> ComputeShader(GlobalShaderMap);
+
+			FComputeShaderUtils::AddPass(
+				GraphBuilder,
+				RDG_EVENT_NAME("NFOR::AddAccumulateTextureRegionPass (%s [%d,%d] -> %s [%d,%d], size:%dx%d)",
+					SourceTexture->Name,
+					SourcePosition.X,
+					SourcePosition.Y,
+					TargetTexture->Name,
+					TargetPosition.X,
+					TargetPosition.Y,
+					Size.X,
+					Size.Y),
+				ERDGPassFlags::Compute,
+				ComputeShader,
+				PassParameters,
+				FComputeShaderUtils::GetGroupCount(Size, NON_LOCAL_MEAN_THREAD_GROUP_SIZE));
+		}
+	}
+
+	void AddAccumulateConstantRegionPass(FRDGBuilder& GraphBuilder, const FLinearColor& ConstantValue, const FRDGTextureRef& TargetTexture, const FRDGTextureRef& Mask,
+		FIntPoint SourcePosition, FIntPoint TargetPosition, FIntPoint Size)
+	{
+		Size = Size == FIntPoint::ZeroValue ? TargetTexture->Desc.Extent : Size;
+		FGlobalShaderMap* GlobalShaderMap = GetGlobalShaderMap(GMaxRHIFeatureLevel);
+		const bool bUseMask = Mask != nullptr;
+
+		typedef FTextureAccumulateConstantCS SHADER;
+		SHADER::FParameters* PassParameters = GraphBuilder.AllocParameters<SHADER::FParameters>();
+		{
+			PassParameters->RWTarget = GraphBuilder.CreateUAV(TargetTexture);
+			PassParameters->TargetPosition = TargetPosition;
+			PassParameters->ConstantValue = ConstantValue;
+			PassParameters->Size = Size;
+			PassParameters->Mask = bUseMask ? GraphBuilder.CreateSRV(Mask) : nullptr;
+		}
+
+		SHADER::FPermutationDomain ComputeShaderPermutationVector;
+		ComputeShaderPermutationVector.Set<SHADER::FDimensionAccumulateByMask>(bUseMask);
+		TShaderMapRef<SHADER> ComputeShader(GlobalShaderMap, ComputeShaderPermutationVector);
+
+		FComputeShaderUtils::AddPass(
+			GraphBuilder,
+			RDG_EVENT_NAME("NFOR::AddAccumulateConstantRegionPass ([%.1f,%.1f,%.1f,%.1f] -> %s [%d,%d], size:%dx%d %s)",
+				ConstantValue.R,
+				ConstantValue.G,
+				ConstantValue.B,
+				ConstantValue.A,
+				TargetTexture->Name,
+				TargetPosition.X,
+				TargetPosition.Y,
+				Size.X,
+				Size.Y,
+				bUseMask ? TEXT("Masked") : TEXT("")),
+			ERDGPassFlags::Compute,
+			ComputeShader,
+			PassParameters,
+			FComputeShaderUtils::GetGroupCount(Size, NON_LOCAL_MEAN_THREAD_GROUP_SIZE));
+	}
+
+	void AddCopyMirroredTexturePass(FRDGBuilder& GraphBuilder, const FRDGTextureRef& SourceTexture, const FRDGTextureRef& TargetTexture,
+		FIntPoint SourcePosition, FIntPoint TargetPosition, FIntPoint Size, bool bAlphaOnly)
+	{
+		const FIntPoint CopySize = Size == FIntPoint::ZeroValue ? TargetTexture->Desc.Extent : Size;
+
+		typedef FCopyTexturePS SHADER;
+
+		SHADER::FParameters* PassParameters = GraphBuilder.AllocParameters<SHADER::FParameters>();
+		{
+			PassParameters->Source = GraphBuilder.CreateSRV(FRDGTextureSRVDesc::Create(SourceTexture));
+			PassParameters->SourceOffset = SourcePosition;
+			PassParameters->TextureSize = SourceTexture->Desc.Extent;
+			PassParameters->RenderTargets[0] = FRenderTargetBinding(TargetTexture, ERenderTargetLoadAction::ENoAction);
+		}
+
+		const FIntRect ViewRect(TargetPosition, TargetPosition + CopySize);
+
+		FGlobalShaderMap* ShaderMap = GetGlobalShaderMap(GMaxRHIFeatureLevel);
+		TShaderMapRef<FCopyTexturePS> PixelShader(ShaderMap);
+
+		FRHIBlendState* BlendState = TStaticBlendState<CW_RGBA, BO_Add, BF_One, BF_Zero, BO_Add, BF_One, BF_Zero>::GetRHI();
+		if (bAlphaOnly)
+		{
+			BlendState = TStaticBlendState<CW_RGBA, BO_Add, BF_Zero, BF_One, BO_Add, BF_One, BF_Zero>::GetRHI();
+		}
+
+		FPixelShaderUtils::AddFullscreenPass(
+			GraphBuilder,
+			ShaderMap,
+			RDG_EVENT_NAME("CopyTexture (%s -> %s,Mirrored%s)",
+				SourceTexture->Name,
+				TargetTexture->Name,
+				bAlphaOnly ? TEXT(" AlphaOnly") : TEXT("")),
+			PixelShader,
+			PassParameters,
+			ViewRect,
+			BlendState
+		);
+	}
+
+	void AddNormalizeRadianceVariancePass(FRDGBuilder& GraphBuilder, const FRDGTextureRef& Albedo, const FRDGTextureRef& RadianceVariance)
+	{
+		FIntPoint Size = RadianceVariance->Desc.Extent;
+		typedef FNormalizeRadianceVarianceByAlbedoCS SHADER;
+		SHADER::FParameters* PassParameters = GraphBuilder.AllocParameters<SHADER::FParameters>();
+		{
+			PassParameters->Albedo = GraphBuilder.CreateSRV(FRDGTextureSRVDesc::Create(Albedo));
+			PassParameters->RWRadianceVariance = GraphBuilder.CreateUAV(RadianceVariance);
+			PassParameters->Size = Size;
+		}
+
+		TShaderMapRef<SHADER> ComputeShader(GetGlobalShaderMap(GMaxRHIFeatureLevel));
+
+		FComputeShaderUtils::AddPass(
+			GraphBuilder,
+			RDG_EVENT_NAME("NFOR::AddNormalizeRadianceVariancePass (%s.RadianceVariance / %s, size:%dx%d)",
+				RadianceVariance->Name,
+				Albedo->Name,
+				Size.X,
+				Size.Y),
+			ERDGPassFlags::Compute,
+			ComputeShader,
+			PassParameters,
+			FComputeShaderUtils::GetGroupCount(Size, NON_LOCAL_MEAN_THREAD_GROUP_SIZE));
+	}
+
+	//--------------------------------------------------------------------------------------------------------------------
+	// Radiance normalization
+	FRDGTextureRef GetPreAlbedoDivideMask(FRDGBuilder& GraphBuilder, const FSceneView& View, const FRDGTextureRef& Normal, const FRDGTextureRef& NormalVariance)
+	{
+		FRDGTextureDesc Desc = Normal->Desc;
+		Desc.Format = PF_R8_UINT;
+		FRDGTextureRef MaskTexture = GraphBuilder.CreateTexture(Desc, TEXT("NFOR.MaskTexture"));
+		{
+			typedef FClassifyPreAlbedoDivideMaskIdCS SHADER;
+			SHADER::FParameters* PassParameters = GraphBuilder.AllocParameters<SHADER::FParameters>();
+			{
+				PassParameters->Normal = GraphBuilder.CreateSRV(Normal);
+				PassParameters->NormalVariance = GraphBuilder.CreateSRV(NormalVariance);
+				PassParameters->TextureSize = Desc.Extent;
+				PassParameters->RWMask = GraphBuilder.CreateUAV(MaskTexture);
+			}
+
+			TShaderMapRef<SHADER> ComputeShader(GetGlobalShaderMap(GMaxRHIFeatureLevel));
+
+			FComputeShaderUtils::AddPass(
+				GraphBuilder,
+				RDG_EVENT_NAME("NFOR::ClassifyPreAlbedoDivideMaskIdCS (size:%dx%d)",
+					Desc.Extent.X,
+					Desc.Extent.Y),
+				ERDGPassFlags::Compute,
+				ComputeShader,
+				PassParameters,
+				FComputeShaderUtils::GetGroupCount(Desc.Extent, NON_LOCAL_MEAN_THREAD_GROUP_SIZE));
+		}
+
+		return MaskTexture;
+	}
+
+	//--------------------------------------------------------------------------------------------------------------------
+	// Non-local mean weight and filtering
+	FNonLocalMeanParameters GetNonLocalMeanParameters(int32 PatchSize, int32 PatchDistance, float Bandwidth)
+	{
+		FNonLocalMeanParameters NonLocalMeanParameters;
+		NonLocalMeanParameters.PatchSize = PatchSize;
+		NonLocalMeanParameters.PatchDistance = PatchDistance;
+		NonLocalMeanParameters.Bandwidth = Bandwidth;
+
+		return NonLocalMeanParameters;
+	}
+
+	FNonLocalMeanParameters GetFeatureNonLocalMeanParameters(float Bandwidth)
+	{
+		const int32 PatchSize = GetNonLocalMeanFeaturePatchSize();
+		const int32 PatchDistance = GetNonLocalMeanFeaturePatchDistance();
+
+		return GetNonLocalMeanParameters(PatchSize, PatchDistance, Bandwidth);
+	}
+
+	void ApplyNonLocalMeanFilter(
+		FRDGBuilder& GraphBuilder,
+		const FNonLocalMeanParameters& NonLocalMeanParameters,
+		const FNFORTextureDesc& Texture,
+		const FNFORTextureDesc& Variance,
+		EVarianceType VarianceTyle,
+		const FRDGTextureRef& FilteredTexture)
+	{
+		FGlobalShaderMap* GlobalShaderMap = GetGlobalShaderMap(GMaxRHIFeatureLevel);
+		FIntPoint TextureSize = Texture.Image->Desc.Extent;
+		const FRDGTextureRef VarianceTexture = Variance.Image ? Variance.Image : GSystemTextures.GetBlackDummy(GraphBuilder);
+
+		typedef FNonLocalMeanFilteringCS SHADER;
+
+		SHADER::FParameters* PassParameters = GraphBuilder.AllocParameters<SHADER::FParameters>();
+		{
+			PassParameters->NLMParams = NonLocalMeanParameters;
+			PassParameters->Image = GraphBuilder.CreateSRV(FRDGTextureSRVDesc::Create(Texture.Image));
+			PassParameters->Variance = GraphBuilder.CreateSRV(FRDGTextureSRVDesc::Create(VarianceTexture));
+			PassParameters->TextureSize = TextureSize;
+			PassParameters->VarianceChannelOffset = Variance.ChannelOffset;
+			PassParameters->DenoisingChannelCount = Texture.ChannelCount;
+			PassParameters->DenoisedImage = GraphBuilder.CreateUAV(FRDGTextureUAVDesc(FilteredTexture));
+		}
+
+		SHADER::FPermutationDomain ComputeShaderPermutationVector;
+		{
+			ComputeShaderPermutationVector.Set<SHADER::FDimensionVarianceType>(VarianceTyle);
+			ComputeShaderPermutationVector.Set<SHADER::FDimensionUseGuide>(false);
+			ComputeShaderPermutationVector.Set<SHADER::FDimensionImageChannelCount>(Texture.NumOfChannel);
+			ComputeShaderPermutationVector.Set<SHADER::FDimPreAlbedoDivide>(GetPreAlbedoDivideRecoverPhase());
+		}
+
+		TShaderMapRef<SHADER> ComputeShader(GlobalShaderMap, ComputeShaderPermutationVector);
+
+		FComputeShaderUtils::AddPass(
+			GraphBuilder,
+			RDG_EVENT_NAME("NFOR::FeatureFiltering(%s, Dim=%d,%d)", Texture.Image->Name, TextureSize.X, TextureSize.Y),
+			ERDGPassFlags::Compute,
+			ComputeShader,
+			PassParameters,
+			FComputeShaderUtils::GetGroupCount(TextureSize, NON_LOCAL_MEAN_THREAD_GROUP_SIZE));
+	}
+
+	void GetNLMWeights(
+		FRDGBuilder& GraphBuilder,
+		const FSceneView& View,
+		const FRadianceDesc& SourceRadiance,
+		const FRadianceDesc& TargetRadiance,
+		const FRDGBufferRef& NonLocalMeanWeightsBuffer,
+		FIntRect Region,
+		const FNonLocalMeanParameters& NonLocalMeanParameters)
+	{
+		const int32 SearchingPatchSize = (NonLocalMeanParameters.PatchDistance * 2 + 1);
+		const int32 NumberOfWeightsPerPixel = SearchingPatchSize * SearchingPatchSize;
+		const FIntPoint TextureSize = SourceRadiance.Data.Image->Desc.Extent;
+		const bool bSeparateSourceTarget = (SourceRadiance.Data.Image != TargetRadiance.Data.Image);
+		
+		// Query the non-local mean weights for the radiance.
+		{
+			typedef FNonLocalMeanWeightsCS SHADER;
+			SHADER::FParameters* PassParameters = GraphBuilder.AllocParameters<SHADER::FParameters>();
+			{
+				PassParameters->CommonParameters.NLMParams = NonLocalMeanParameters;
+				PassParameters->CommonParameters.Image = GraphBuilder.CreateSRV(FRDGTextureSRVDesc::Create(SourceRadiance.Data.Image));
+				PassParameters->CommonParameters.Variance = GraphBuilder.CreateSRV(FRDGTextureSRVDesc::Create(SourceRadiance.Variance.Image));
+				PassParameters->CommonParameters.TextureSize = TextureSize;
+				PassParameters->CommonParameters.VarianceChannelOffset = SourceRadiance.Variance.ChannelOffset;
+
+				PassParameters->NonLocalMeanWeights = GraphBuilder.CreateUAV(NonLocalMeanWeightsBuffer, PF_R32_FLOAT);
+				PassParameters->Region = Region;
+
+				if (bSeparateSourceTarget)
+				{
+					PassParameters->TargetImage = GraphBuilder.CreateSRV(FRDGTextureSRVDesc::Create(TargetRadiance.Data.Image));
+					PassParameters->TargetVariance = GraphBuilder.CreateSRV(FRDGTextureSRVDesc::Create(TargetRadiance.Variance.Image));
+				}
+			}
+
+			SHADER::FPermutationDomain ComputeShaderPermutationVector;
+			{
+				ComputeShaderPermutationVector.Set<SHADER::FDimensionVarianceType>(EVarianceType::GreyScale);
+				ComputeShaderPermutationVector.Set<SHADER::FDimensionUseGuide>(false);
+				ComputeShaderPermutationVector.Set<SHADER::FDimensionImageChannelCount>(SourceRadiance.Data.NumOfChannel);
+				ComputeShaderPermutationVector.Set<SHADER::FDimensionSeparateSourceTarget>(bSeparateSourceTarget);
+				ComputeShaderPermutationVector.Set<SHADER::FDimPreAlbedoDivide>(GetPreAlbedoDivideRecoverPhase());
+			}
+
+			TShaderMapRef<SHADER> ComputeShader(GetGlobalShaderMap(GMaxRHIFeatureLevel), ComputeShaderPermutationVector);
+
+			FComputeShaderUtils::AddPass(
+				GraphBuilder,
+				RDG_EVENT_NAME("NFOR::NonLocalMeanWeights (Rect=(%d,%d,%d,%d),ps=%d,pd=%d,bw=%.2f)",
+					Region.Min.X,
+					Region.Min.Y,
+					Region.Max.X,
+					Region.Max.Y,
+					NonLocalMeanParameters.PatchSize,
+					NonLocalMeanParameters.PatchDistance,
+					NonLocalMeanParameters.Bandwidth),
+				ERDGPassFlags::Compute,
+				ComputeShader,
+				PassParameters,
+				FComputeShaderUtils::GetGroupCount(Region.Size(), NON_LOCAL_MEAN_THREAD_GROUP_SIZE));
+		}
+	}
+
+	//--------------------------------------------------------------------------------------------------------------------
+	// Collaborative filtering
+	//	1. Tiling
+	void AddCopyTextureToBufferPass(
+		FRDGBuilder& GraphBuilder,
+		const FRDGTextureRef Source,
+		const FRDGBufferRef Dest,
+		int32 CopyChannelOffset,
+		int32 CopyChannelCount,
+		int32 NumberOfSourceChannel,
+		int32 BufferChannelOffset,
+		int32 BufferChannelSize,
+		FIntRect CopyRegion)
+	{
+		FGlobalShaderMap* GlobalShaderMap = GetGlobalShaderMap(GMaxRHIFeatureLevel);
+		const int SourceChannelCount = NumberOfSourceChannel;
+		FIntPoint TextureSize = Source->Desc.Extent;
+
+		{
+			typedef FCopyTextureToBufferCS SHADER;
+			SHADER::FParameters* PassParameters = GraphBuilder.AllocParameters<SHADER::FParameters>();
+			{
+				PassParameters->Source = GraphBuilder.CreateSRV(FRDGTextureSRVDesc::Create(Source));
+				PassParameters->Dest = GraphBuilder.CreateUAV(Dest, PF_R32_FLOAT);
+				PassParameters->TextureSize = TextureSize;
+				PassParameters->CopyChannelOffset = CopyChannelOffset;
+				PassParameters->CopyChannelCount = CopyChannelCount;
+				PassParameters->BufferChannelOffset = BufferChannelOffset;
+				PassParameters->BufferChannelSize = BufferChannelSize;
+				PassParameters->CopyRegion = CopyRegion;
+			}
+
+			SHADER::FPermutationDomain ComputeShaderPermutationVector;
+			ComputeShaderPermutationVector.Set<SHADER::FDimensionSourceChannelCount>(SourceChannelCount);
+			TShaderMapRef<SHADER> ComputeShader(GlobalShaderMap, ComputeShaderPermutationVector);
+
+			FComputeShaderUtils::AddPass(
+				GraphBuilder,
+				RDG_EVENT_NAME("NFOR::CopyTextureToBuffer (Dim=%dx%d,s:%d:%d -> b:%d)",
+					TextureSize.X,
+					TextureSize.Y,
+					CopyChannelOffset,
+					CopyChannelOffset + CopyChannelCount - 1,
+					BufferChannelOffset),
+				ERDGPassFlags::Compute,
+				ComputeShader,
+				PassParameters,
+				FComputeShaderUtils::GetGroupCount(CopyRegion.Size(), NON_LOCAL_MEAN_THREAD_GROUP_SIZE));
+		}
+	}
+
+	void AddCopyTextureToBufferPass(
+		FRDGBuilder& GraphBuilder,
+		const FRDGTextureRef Source,
+		const FRDGBufferRef Dest,
+		int32 CopyChannelOffset,
+		int32 CopyChannelCount,
+		int32 NumberOfSourceChannel,
+		int32 BufferChannelOffset,
+		int32 BufferChannelSize)
+	{
+		AddCopyTextureToBufferPass(
+			GraphBuilder,
+			Source,
+			Dest,
+			CopyChannelOffset,
+			CopyChannelCount,
+			NumberOfSourceChannel,
+			BufferChannelOffset,
+			BufferChannelSize,
+			FIntRect(FIntPoint(0, 0), Source->Desc.Extent));
+	}
+
+	void AddNormalizeTexturePass( FRDGBuilder& GraphBuilder, const FRDGTextureRef& InputTexture)
+	{
+		FIntPoint TextureSize = InputTexture->Desc.Extent;
+		FGlobalShaderMap* GlobalShaderMap = GetGlobalShaderMap(GMaxRHIFeatureLevel);
+
+		{
+			typedef FNormalizeTextureCS SHADER;
+			SHADER::FParameters* PassParameters = GraphBuilder.AllocParameters<SHADER::FParameters>();
+			{
+				PassParameters->RWSource = GraphBuilder.CreateUAV(InputTexture);
+				PassParameters->TextureSize = TextureSize;
+			}
+
+			TShaderMapRef<SHADER> ComputeShader(GlobalShaderMap);
+
+			FComputeShaderUtils::AddPass(
+				GraphBuilder,
+				RDG_EVENT_NAME("NFOR::AddNormalizeTexturePass (%dx%d)",
+					TextureSize.X,
+					TextureSize.Y),
+				ERDGPassFlags::Compute,
+				ComputeShader,
+				PassParameters,
+				FComputeShaderUtils::GetGroupCount(TextureSize, NON_LOCAL_MEAN_THREAD_GROUP_SIZE));
+		}
+	}
+
+
+	//	2. Weighted Least-square solver
+	FRDGBufferRef RegressionKernel::AllocateMatrixfBuffer(FRDGBuilder& GraphBuilder, int32 NumOfMatrices, int32 Dim0, int32 Dim1, const TCHAR* Name)
+	{
+		const int32 BytesPerElement = sizeof(float);
+
+		FRDGBufferDesc BufferDesc =
+			FRDGBufferDesc::CreateBufferDesc(BytesPerElement, NumOfMatrices * Dim0 * Dim1);
+
+		return GraphBuilder.CreateBuffer(BufferDesc, Name);
+	}
+
+	FRDGBufferRef RegressionKernel::FInPlaceBatchedMatrixMultiplicationCS::AllocateResultBuffer( FRDGBuilder& GraphBuilder, FIntPoint Size, int32 F, int32 A)
+	{
+		return AllocateMatrixfBuffer(GraphBuilder, Size.X * Size.Y, F, A, TEXT("NFOR.Matrix.Result"));
+	}
+
+	FRDGBufferRef ApplyBatchedInPlaceMatrixMultiplication(
+		FRDGBuilder& GraphBuilder,
+		FRDGBufferRef X,
+		FIntPoint XDim,
+		FRDGBufferRef W,
+		int32 WDim,
+		FIntPoint TextureSize,
+		int PatchDistance,
+		RegressionKernel::EWeightedMultiplicationType MultiplicationType = RegressionKernel::EWeightedMultiplicationType::Quadratic,
+		FRDGBufferRef Y = nullptr,
+		FIntPoint YDim = 0)
+	{
+		FRDGBufferRef ResultMatrix = nullptr;
+		{
+			typedef RegressionKernel::FInPlaceBatchedMatrixMultiplicationCS SHADER;
+			const bool GeneralizedMultiplication = MultiplicationType == RegressionKernel::EWeightedMultiplicationType::Generalized;
+			const bool bFeatureAddConstant = ShouldFeatureAddConstant();
+			const int BufferXDimWithConstant = XDim.Y + (bFeatureAddConstant ? 1 : 0);
+			FIntPoint ResultMatrixDimension = FIntPoint(BufferXDimWithConstant, GeneralizedMultiplication ? YDim.Y : BufferXDimWithConstant);
+			int32 SamplingStep = GetSamplingStep(WDim, ResultMatrixDimension.X * ResultMatrixDimension.Y);
+
+			ResultMatrix = SHADER::AllocateResultBuffer(GraphBuilder, TextureSize, ResultMatrixDimension.X, ResultMatrixDimension.Y);
+
+			SHADER::FParameters* PassParameters = GraphBuilder.AllocParameters<SHADER::FParameters>();
+			{
+				PassParameters->X = GraphBuilder.CreateSRV(FRDGBufferSRVDesc(X, PF_R32_FLOAT));
+				PassParameters->XDim = XDim;
+
+				PassParameters->W = GraphBuilder.CreateSRV(FRDGBufferSRVDesc(W, PF_R32_FLOAT));
+				PassParameters->WDim = WDim;
+
+
+				if (GeneralizedMultiplication)
+				{
+					PassParameters->Y = GraphBuilder.CreateSRV(FRDGBufferSRVDesc(Y, PF_R32_FLOAT));
+					PassParameters->YDim = YDim;
+				}
+				else
+				{
+					PassParameters->Y = GraphBuilder.CreateSRV(FRDGBufferSRVDesc(X, PF_R32_FLOAT));
+					PassParameters->YDim = XDim;
+				}
+
+				PassParameters->TextureSize = TextureSize;
+				PassParameters->PatchDistance = PatchDistance;
+				PassParameters->NumOfWeigthsPerPixelPerFrame = (PatchDistance * 2 + 1) * (PatchDistance * 2 + 1);
+				PassParameters->NumOfTemporalFrames = WDim / PassParameters->NumOfWeigthsPerPixelPerFrame;
+				PassParameters->SourceFrameIndex = GetDenoisingFrameIndex(PassParameters->NumOfTemporalFrames);
+				PassParameters->SamplingStep = SamplingStep;
+				PassParameters->Result = GraphBuilder.CreateUAV(FRDGBufferUAVDesc(ResultMatrix, PF_R32_FLOAT));
+			}
+
+			//TODO: clean up mutation.
+			SHADER::FPermutationDomain ComputeShaderPermutationVector;
+			{
+				ComputeShaderPermutationVector.Set<SHADER::FDimWeightedMultiplicationType>(MultiplicationType);
+				ComputeShaderPermutationVector.Set<SHADER::FDimAddConstantFeatureDim>(bFeatureAddConstant);
+				ComputeShaderPermutationVector.Set<SHADER::FDimOptimizeTargetMatrixMultiplication>(true);
+				ComputeShaderPermutationVector.Set<SHADER::FDimNumFeature>(BufferXDimWithConstant);
+				ComputeShaderPermutationVector.Set<SHADER::FDimUseSamplingStep>(SamplingStep > 1);
+			}
+
+			TShaderMapRef<SHADER> ComputeShader(GetGlobalShaderMap(GMaxRHIFeatureLevel), ComputeShaderPermutationVector);
+
+			FComputeShaderUtils::AddPass(
+				GraphBuilder,
+				RDG_EVENT_NAME("NFOR::Matrix Multiplication (%s%s)",
+					RegressionKernel::GetEventName(MultiplicationType),
+					bFeatureAddConstant? TEXT(" +Const. Feature":TEXT(""))),
+				ERDGPassFlags::Compute,
+				ComputeShader,
+				PassParameters,
+				FComputeShaderUtils::GetGroupCount(TextureSize, SHADER::GetThreadGroupSize()));
+		}
+
+		return ResultMatrix;
+	}
+
+	void ReconstructByFrame(
+		FRDGBuilder& GraphBuilder,
+		FRDGBufferRef Feature,
+		FRDGBufferRef ReconstructionWeights,
+		FRDGBufferRef NonLocalMeanWeightsBuffer,
+		FRDGTextureRef FilteredRadiance,
+		FRDGTextureRef SourceAlbedo,
+		FWeightedLSRDesc WeightedLSRDesc,
+		int FrameIndex,
+		RegressionKernel::FReconstructSpatialTemporalImage::EReconstructionType ReconstructionType = RegressionKernel::FReconstructSpatialTemporalImage::EReconstructionType::Scatter
+	)
+	{
+
+		checkf(FrameIndex < WeightedLSRDesc.NumOfFrames,
+			TEXT("FrameIndex should be less than total number of frames: %d < %d failed:"), FrameIndex, WeightedLSRDesc.NumOfFrames);
+		
+		FRDGTextureRef ReconstructionBuffer64 = nullptr;
+		FRDGBufferRef ReconstructionBuffer = nullptr;
+
+		FIntPoint TextureSize = FilteredRadiance->Desc.Extent;
+
+		{
+			typedef RegressionKernel::FReconstructSpatialTemporalImage SHADER;
+			if (ReconstructionType == SHADER::EReconstructionType::Scatter)
+			{
+				FRDGBufferDesc BufferDesc = FRDGBufferDesc::CreateStructuredDesc(sizeof(FVector4f),TextureSize.X * TextureSize.Y);
+				ReconstructionBuffer = GraphBuilder.CreateBuffer(BufferDesc, TEXT("NFOR.WeightedLSR.ReconstructionBuffer"));
+
+				const EPixelFormat PixelFormat64 = GPixelFormats[PF_R64_UINT].Supported ? PF_R64_UINT : PF_R32G32_UINT;
+
+				const FRDGTextureDesc ReconstructionBufferDesc = FRDGTextureDesc::Create2D(
+					FIntPoint(TextureSize.X*4, TextureSize.Y),
+					PixelFormat64,
+					FClearValueBinding::None,
+					TexCreate_RenderTargetable|TexCreate_ShaderResource | TexCreate_UAV | ETextureCreateFlags::Atomic64Compatible);
+
+				ReconstructionBuffer64 = GraphBuilder.CreateTexture(ReconstructionBufferDesc, TEXT("NFOR.WeightedLSR.ReconstructionBuffer64"));
+
+				AddClearUAVPass(GraphBuilder, GraphBuilder.CreateUAV(ReconstructionBuffer), 0);
+				AddClearRenderTargetPass(GraphBuilder, ReconstructionBuffer64, FLinearColor::Transparent);
+			}
+
+			const bool bFeatureAddConstant = ShouldFeatureAddConstant();
+			const int32 NumOfAdditionalFeatures = (bFeatureAddConstant ? 1 : 0);
+
+			FIntPoint XDimension = FIntPoint(WeightedLSRDesc.NumOfWeightsPerPixel, WeightedLSRDesc.NumOfFeatureChannelsPerFrame);
+			const int PatchDistance = (FMath::Sqrt(float(WeightedLSRDesc.NumOfWeightsPerPixel / WeightedLSRDesc.NumOfFrames)) - 1) / 2;
+			const int TotalNumOfFeaturesPerFrame = WeightedLSRDesc.NumOfFeatureChannelsPerFrame + NumOfAdditionalFeatures;
+			FIntPoint BDim = FIntPoint(TotalNumOfFeaturesPerFrame, WeightedLSRDesc.NumOfRadianceChannelsPerFrame);
+
+			SHADER::FParameters* PassParameters = GraphBuilder.AllocParameters<SHADER::FParameters>();
+			{
+				PassParameters->X = GraphBuilder.CreateSRV(FRDGBufferSRVDesc(Feature, PF_R32_FLOAT));
+				PassParameters->W = GraphBuilder.CreateSRV(FRDGBufferSRVDesc(NonLocalMeanWeightsBuffer, PF_R32_FLOAT));
+				PassParameters->B = GraphBuilder.CreateSRV(FRDGBufferSRVDesc(ReconstructionWeights, PF_R32_FLOAT));
+				PassParameters->RWReconstruction = GraphBuilder.CreateUAV(FilteredRadiance);
+				PassParameters->RWReconstructBuffer = ReconstructionBuffer ? GraphBuilder.CreateUAV(ReconstructionBuffer) : nullptr;
+				PassParameters->RWReconstructBuffer64 = ReconstructionBuffer ? GraphBuilder.CreateUAV(ReconstructionBuffer64) : nullptr;
+
+				PassParameters->XDim = XDimension;
+				PassParameters->WDim = WeightedLSRDesc.NumOfWeightsPerPixel;
+				PassParameters->BDim = BDim;
+
+				PassParameters->TextureSize = FIntPoint(WeightedLSRDesc.Width, WeightedLSRDesc.Height);
+				PassParameters->PatchDistance = PatchDistance;
+				PassParameters->FrameIndex = FrameIndex;
+
+				PassParameters->NumOfTemporalFrames = WeightedLSRDesc.NumOfFrames;
+				PassParameters->NumOfWeigthsPerPixelPerFrame = WeightedLSRDesc.NumOfWeightsPerPixel / WeightedLSRDesc.NumOfFrames;
+			}
+			
+			SHADER::FPermutationDomain ComputeShaderPermutationVector;
+			{
+				ComputeShaderPermutationVector.Set<SHADER::FDimReconstructionType>(ReconstructionType);
+				ComputeShaderPermutationVector.Set<SHADER::FDimPreAlbedoDivide>(GetPreAlbedoDivideRecoverPhase());
+				ComputeShaderPermutationVector.Set<SHADER::FDimNumFeature>(BDim.X);
+			}
+
+			TShaderMapRef<SHADER> ComputeShader(GetGlobalShaderMap(GMaxRHIFeatureLevel), ComputeShaderPermutationVector);
+
+			FComputeShaderUtils::AddPass(
+				GraphBuilder,
+				RDG_EVENT_NAME("NFOR::Reconstruction(T=%d,%s)",
+					FrameIndex,
+					SHADER::GetEventName(ReconstructionType)),
+				ERDGPassFlags::Compute,
+				ComputeShader,
+				PassParameters,
+				FComputeShaderUtils::GetGroupCount(FIntPoint(WeightedLSRDesc.Width, WeightedLSRDesc.Height), NON_LOCAL_MEAN_THREAD_GROUP_SIZE));
+		}
+
+		if (ReconstructionBuffer)
+		{
+			typedef FAccumulateBufferToTextureCS SHADER;
+			SHADER::FParameters* PassParameters = GraphBuilder.AllocParameters<SHADER::FParameters>();
+			{
+				PassParameters->StructuredBufferSource = GraphBuilder.CreateSRV(FRDGBufferSRVDesc(ReconstructionBuffer));
+				PassParameters->ReconstructBuffer64 = GraphBuilder.CreateSRV(ReconstructionBuffer64);
+				PassParameters->RWTarget = GraphBuilder.CreateUAV(FilteredRadiance);
+				PassParameters->TextureSize = TextureSize;
+			}
+
+			SHADER::FPermutationDomain ComputeShaderPermutationVector;
+			{
+				ComputeShaderPermutationVector.Set<SHADER::FDimPreAlbedoDivide>(GetPreAlbedoDivideRecoverPhase());
+			}
+
+			TShaderMapRef<SHADER> ComputeShader(GetGlobalShaderMap(GMaxRHIFeatureLevel), ComputeShaderPermutationVector);
+
+			FComputeShaderUtils::AddPass(
+				GraphBuilder,
+				RDG_EVENT_NAME("NFOR::AccumulateBufferToTexture(%dx%d)", TextureSize.X, TextureSize.Y),
+				ERDGPassFlags::Compute,
+				ComputeShader,
+				PassParameters,
+				FComputeShaderUtils::GetGroupCount(TextureSize, NON_LOCAL_MEAN_THREAD_GROUP_SIZE));
+		}
+	}
+
+	void SolveWeightedLSR(
+		FRDGBuilder& GraphBuilder,
+		const FSceneView& View,
+		const FRDGBufferRef& Feature,
+		const FRDGTextureRef& Radiance,
+		const FRDGBufferRef& NonLocalMeanWeightsBuffer,
+		const FRDGTextureRef& FilteredRadiance,
+		FWeightedLSRDesc WeightedLSRDesc,
+		const FRDGBufferRef Radiances,
+		const FRDGTextureRef& SourceAlbedo
+	)
+	{
+		FIntPoint TextureSize = Radiance->Desc.Extent;
+		
+
+		if (GetRegressionDevice() == ERegressionDevice::CPU)
+		{
+			SolveWeightedLSRCPU(
+				GraphBuilder,
+				View,
+				Feature,
+				Radiance,
+				NonLocalMeanWeightsBuffer,
+				FilteredRadiance,
+				WeightedLSRDesc,
+				Radiances,
+				SourceAlbedo);
+
+			return;
+		}
+
+		checkf(WeightedLSRDesc.SolverType == EWeightedLSRSolverType::Tiled,TEXT("The weighted LSR solver should select tiled"));
+
+		FIntPoint XDimension = FIntPoint(WeightedLSRDesc.NumOfWeightsPerPixel, WeightedLSRDesc.NumOfFeatureChannelsPerFrame);
+		int PatchDistance = (FMath::Sqrt(float(WeightedLSRDesc.NumOfWeightsPerPixel / WeightedLSRDesc.NumOfFrames)) - 1) / 2;
+
+		// 1. Process the data into A, B for Ax=B.
+		FRDGBufferRef AMatrix = nullptr;
+		FRDGBufferRef BMatrix = nullptr;
+		{
+			AMatrix = ApplyBatchedInPlaceMatrixMultiplication(
+				GraphBuilder,
+				Feature,
+				XDimension,
+				NonLocalMeanWeightsBuffer,
+				WeightedLSRDesc.NumOfWeightsPerPixel,
+				FIntPoint(WeightedLSRDesc.Width, WeightedLSRDesc.Height),
+				PatchDistance,
+				RegressionKernel::EWeightedMultiplicationType::Quadratic);
+
+			BMatrix = ApplyBatchedInPlaceMatrixMultiplication(
+				GraphBuilder,
+				Feature,
+				XDimension,
+				NonLocalMeanWeightsBuffer,
+				WeightedLSRDesc.NumOfWeightsPerPixel,
+				FIntPoint(WeightedLSRDesc.Width, WeightedLSRDesc.Height),
+				PatchDistance,
+				RegressionKernel::EWeightedMultiplicationType::Generalized,
+				Radiances,
+				FIntPoint(WeightedLSRDesc.NumOfWeightsPerPixel, WeightedLSRDesc.NumOfRadianceChannelsPerFrame));
+		}
+
+		// 2. Solve the linear equation Ax=B.
+		const bool bFeatureAddConstant = ShouldFeatureAddConstant();
+		const int32 NumOfAdditionalFeatures = bFeatureAddConstant ? 1 : 0;
+		const int32 NumOfElementsPerRow = WeightedLSRDesc.Width;
+		const int32 NumOfElements = WeightedLSRDesc.Width * WeightedLSRDesc.Height;
+		const int32 TotalNumOfFeaturesPerFrame = WeightedLSRDesc.NumOfFeatureChannelsPerFrame + NumOfAdditionalFeatures;
+
+		FIntPoint BDim = FIntPoint(TotalNumOfFeaturesPerFrame, WeightedLSRDesc.NumOfRadianceChannelsPerFrame);
+		FRDGBufferRef ReconstructionWeights = RegressionKernel::AllocateMatrixfBuffer(
+			GraphBuilder, 
+			NumOfElements,
+			BDim.X,
+			BDim.Y, TEXT("NFOR.WeightedLSR.ReconstructWeights"));
+
+		if (GetLinearSolverDevice() == ELinearSolverDevice::CPU)
+		{
+			SolveLinearEquationCPU(
+				GraphBuilder,
+				AMatrix,
+				BMatrix,
+				NumOfElements,
+				BDim,
+				ReconstructionWeights);
+		}
+		else
+		{
+			typedef RegressionKernel::FLinearSolverCS SHADER;
+			SHADER::FParameters* PassParameters = GraphBuilder.AllocParameters<SHADER::FParameters>();
+			PassParameters->A = GraphBuilder.CreateSRV(FRDGBufferSRVDesc(AMatrix, PF_R32_FLOAT));
+			PassParameters->ADim = FIntPoint(TotalNumOfFeaturesPerFrame, TotalNumOfFeaturesPerFrame);
+			PassParameters->B = GraphBuilder.CreateSRV(FRDGBufferSRVDesc(BMatrix, PF_R32_FLOAT));
+			PassParameters->BDim = BDim;
+			PassParameters->Result = GraphBuilder.CreateUAV(FRDGBufferUAVDesc(ReconstructionWeights, PF_R32_FLOAT));
+
+			PassParameters->NumOfElements = NumOfElements;
+			PassParameters->NumOfElementsPerRow = NumOfElementsPerRow;
+
+			SHADER::FPermutationDomain ComputeShaderPermutationVector;
+
+			checkf(BDim.X >= 6 && BDim.X <= 8, TEXT("Number of features should be between 6 and 8"));
+			ComputeShaderPermutationVector.Set<SHADER::FDimNumFeature>(BDim.X);
+
+			TShaderMapRef<SHADER> ComputeShader(GetGlobalShaderMap(GMaxRHIFeatureLevel), ComputeShaderPermutationVector);
+
+			FComputeShaderUtils::AddPass(
+				GraphBuilder,
+				RDG_EVENT_NAME("NFOR::BatchedLinearSolver(F=%d, C=%d)", BDim.X, BDim.Y),
+				ERDGPassFlags::Compute,
+				ComputeShader,
+				PassParameters,
+				FComputeShaderUtils::GetGroupCount(FIntPoint(WeightedLSRDesc.Width, WeightedLSRDesc.Height), NON_LOCAL_MEAN_THREAD_GROUP_SIZE));
+		}
+
+		// 3. Reconstruct 
+		{
+			AddClearUAVPass(GraphBuilder,GraphBuilder.CreateUAV(FilteredRadiance), 0.0f, ERDGPassFlags::Compute);
+
+			const int32 ReconstructDebugFrameIndex = GetReconstructionDebugFrameIndex();
+			const bool IsReconstructDebugEnabled = ReconstructDebugFrameIndex >= 0;
+
+			for (int32 FrameIndex = 0; FrameIndex < WeightedLSRDesc.NumOfFrames; ++FrameIndex)
+			{
+				if (IsReconstructDebugEnabled)
+				{
+					FrameIndex = FMath::Min(ReconstructDebugFrameIndex, WeightedLSRDesc.NumOfFrames - 1);
+				}
+
+				RegressionKernel::FReconstructSpatialTemporalImage::EReconstructionType
+					ReconstructionType = GetReconstructionType(FrameIndex, GetDenoisingFrameIndex(WeightedLSRDesc.NumOfFrames));
+				
+				ReconstructByFrame(
+					GraphBuilder,
+					Feature,
+					ReconstructionWeights,
+					NonLocalMeanWeightsBuffer,
+					FilteredRadiance,
+					SourceAlbedo,
+					WeightedLSRDesc,
+					FrameIndex,
+					ReconstructionType);
+
+				if (IsReconstructDebugEnabled)
+				{
+					break;
+				}
+			}
+		}
+
+		// Multiply albedo if the albedo recover phase is final.
+		if (GetPreAlbedoDivideRecoverPhase()== EAlbedoDivideRecoverPhase::Final)
+		{
+			FIntPoint SourcePosition = WeightedLSRDesc.TileStartPosition - WeightedLSRDesc.Offset;
+			AddMultiplyTextureRegionPass(GraphBuilder, SourceAlbedo, FilteredRadiance, true, SourcePosition, FIntPoint::ZeroValue, WeightedLSRDesc.TextureSize);
+		}
+	}
+
+	int32 GetNumOfCombinedFeatureChannels(const TArray<FFeatureDesc>& FeatureDescs) {
+		int NumOfChannels = 0;
+		for (FFeatureDesc FeatureDesc : FeatureDescs)
+		{
+			if (FeatureDesc.Feature.Image)
+			{
+				NumOfChannels += FeatureDesc.Feature.ChannelCount;
+			}
+		}
+		return NumOfChannels;
+	};
+
+	FRDGTextureRef CollaborativeRegression(
+		FRDGBuilder& GraphBuilder,
+		const FSceneView& View,
+		const TArray<FRadianceDesc>& Radiances,
+		const TArray<FFeatureDesc>& Features,
+		const FNonLocalMeanParameters& RadianceNonLocalMeanParameters)
+	{
+		const int32 NumOfFeatures = Features.Num();
+		const int32 NumOfRadiances = Radiances.Num();
+		const int32 SourceIndex = GetDenoisingFrameIndex(NumOfRadiances); // The current denoising frame data index
+
+		const int SearchingPatchSize = (RadianceNonLocalMeanParameters.PatchDistance * 2 + 1);
+		const int NumberOfWeightsPerPixel = SearchingPatchSize * SearchingPatchSize;
+
+		// TODO: adaptive tile size for best performance.
+		const int32 NumOfTilesOneSide = GetNumOfTiles();
+		FIntPoint NumOfTiles = FIntPoint(NumOfTilesOneSide, NumOfTilesOneSide);
+		const int32 TotalTileCount = NumOfTilesOneSide * NumOfTilesOneSide;
+
+		FIntPoint TextureSize = Radiances[0].Data.Image->Desc.Extent;
+		FIntPoint TileSize = FMath::DivideAndRoundUp(TextureSize, NumOfTiles);
+		FIntPoint PaddingTileOffset = RadianceNonLocalMeanParameters.PatchDistance;
+		FIntPoint PaddedTileSize = TileSize + PaddingTileOffset * 2;
+		FIntRect  PaddedTileRect = FIntRect(FIntPoint(0,0), PaddedTileSize);
+		const int NumOfCombinedFeatureChannels = GetNumOfCombinedFeatureChannels(Features);
+
+		const int32 NonLocalMeanSingleFrameWeightSize = TileSize.X * TileSize.Y * NumberOfWeightsPerPixel;
+		const int32 BytesPerElement = sizeof(float);
+		FRDGBufferDesc NonLocalMeanSingleFrameWeightsBufferDesc = 
+			FRDGBufferDesc::CreateBufferDesc(BytesPerElement, NonLocalMeanSingleFrameWeightSize);
+		FRDGBufferRef NonLocalMeanSingleFrameWeightsBuffer = GraphBuilder.CreateBuffer(NonLocalMeanSingleFrameWeightsBufferDesc, TEXT("NFOR.NLMSingleFrameWeightsBuffer"));
+		
+		FRDGBufferRef NonLocalMeanWeightsBuffer = nullptr;
+		if (NumOfRadiances > 1)
+		{
+			FRDGBufferDesc NonLocalFrameWeightsBufferDesc =
+				FRDGBufferDesc::CreateBufferDesc(BytesPerElement, NonLocalMeanSingleFrameWeightSize * NumOfRadiances);
+			NonLocalMeanWeightsBuffer = GraphBuilder.CreateBuffer(NonLocalFrameWeightsBufferDesc, TEXT("NFOR.NLMWeightsBuffer"));
+		}
+		else
+		{
+			NonLocalMeanWeightsBuffer = NonLocalMeanSingleFrameWeightsBuffer;
+		}
+
+		FRDGBufferDesc CombinedFeatureDesc = FRDGBufferDesc::CreateBufferDesc(BytesPerElement, PaddedTileSize.X * PaddedTileSize.Y * NumOfCombinedFeatureChannels);
+		FRDGBufferRef CombinedFeatures = GraphBuilder.CreateBuffer(CombinedFeatureDesc, TEXT("NFOR.CombinedFeatures"));
+
+		const int32 NumOfCombinedRadianceChannels = GetNumOfCombinedFeatureChannels(Radiances);
+		FRDGBufferDesc CombinedRadianceDesc = FRDGBufferDesc::CreateBufferDesc(BytesPerElement * 4, PaddedTileSize.X * PaddedTileSize.Y * NumOfRadiances);
+		FRDGBufferRef CombinedRadiances = GraphBuilder.CreateBuffer(CombinedRadianceDesc, TEXT("NFOR.CombinedRadiances"));
+
+		FRDGTextureDesc FilteredRadianceDesc = Radiances[SourceIndex].Data.Image->Desc;
+		FilteredRadianceDesc.Flags |= TexCreate_RenderTargetable;
+		FRDGTextureRef FilteredRadiance = GraphBuilder.CreateTexture(FilteredRadianceDesc, TEXT("NFOR.FilteredRadiance"));
+
+		FRDGTextureDesc RadianceTileDesc = FilteredRadianceDesc;
+		RadianceTileDesc.Extent = PaddedTileSize;
+		FRDGTextureRef RadianceTileTexture = GraphBuilder.CreateTexture(RadianceTileDesc, TEXT("NFOR.RadianceTile"));
+		FRDGTextureRef DenoisedTileTexture = GraphBuilder.CreateTexture(RadianceTileDesc, TEXT("NFOR.DenoisedRadianceTile"));
+		
+		AddClearUAVPass(GraphBuilder, GraphBuilder.CreateUAV(FilteredRadiance), FLinearColor::Transparent, ERDGPassFlags::Compute);
+		
+		auto GetTileIndex = [TotalTileCount, NumOfTilesOneSide](int32 Index)
+		{
+			// TODO: generate tiles first and then iterate on tiles?
+			if (IsTileDebugEnabled())
+			{
+				if (GetTileDebugIndex() < 0)
+				{
+					Index = TotalTileCount / 2 - NumOfTilesOneSide / 2;
+				}
+				else
+				{
+					Index = GetTileDebugIndex();
+				}
+			}
+			return Index;
+		};
+
+		//TODO: Each tile can be parallelized.
+		for (int32 i = 0; i < TotalTileCount; ++i)
+		{
+			int32 TileIndex = GetTileIndex(i);
+
+			FIntPoint TileStartPoint = FIntPoint(TileIndex % NumOfTiles.X, TileIndex / NumOfTiles.X) * TileSize;
+			FIntRect TileRegion = FIntRect(FIntPoint(0), TileSize) + TileStartPoint;
+			FIntRect  PaddedTileRegion = PaddedTileRect + TileStartPoint - PaddingTileOffset;
+
+			// Get the weights W
+			for (int32 RadianceId = 0; RadianceId < NumOfRadiances; ++RadianceId)
+			{
+				GetNLMWeights(
+					GraphBuilder,
+					View,
+					Radiances[SourceIndex],
+					Radiances[RadianceId],
+					NonLocalMeanSingleFrameWeightsBuffer,
+					TileRegion,
+					RadianceNonLocalMeanParameters);
+
+				if (NumOfRadiances > 1)
+				{
+					AddCopyBufferPass(GraphBuilder, NonLocalMeanWeightsBuffer, NonLocalMeanSingleFrameWeightSize * BytesPerElement * RadianceId, 
+						NonLocalMeanSingleFrameWeightsBuffer, 0, NonLocalMeanSingleFrameWeightSize * BytesPerElement);
+				}
+			}
+
+			// Get raw color Y
+			{
+				int32 BufferChannelOffset = 0;
+				
+				for (int32 RadianceId = 0; RadianceId < NumOfRadiances; ++RadianceId)
+				{
+					AddCopyMirroredTexturePass(GraphBuilder, Radiances[RadianceId].Data.Image, RadianceTileTexture, PaddedTileRegion.Min, FIntPoint::ZeroValue, PaddedTileSize);
+
+					FNFORTextureDesc Texture = Radiances[RadianceId].Data;
+
+					AddCopyTextureToBufferPass(GraphBuilder, Texture.Image, CombinedRadiances,
+						Texture.ChannelOffset,
+						Texture.ChannelCount,
+						Texture.NumOfChannel,
+						BufferChannelOffset,
+						NumOfCombinedRadianceChannels,
+						PaddedTileRegion);
+
+					BufferChannelOffset += Texture.ChannelCount;
+				}
+
+				checkf(BufferChannelOffset == NumOfCombinedRadianceChannels, TEXT("Number of channels used by radiances does not match the channel count in the buffer."));
+			}
+
+			// Get the feature vector X
+			{
+				int32 BufferChannelOffset = 0;
+
+				for (int32 FeatureId = 0; FeatureId < NumOfFeatures; ++FeatureId)
+				{
+					FNFORTextureDesc Texture = Features[FeatureId].Data;
+
+					AddCopyTextureToBufferPass(GraphBuilder, Texture.Image, CombinedFeatures,
+						Texture.ChannelOffset,
+						Texture.ChannelCount,
+						Texture.NumOfChannel,
+						BufferChannelOffset,
+						NumOfCombinedFeatureChannels,
+						PaddedTileRegion);
+
+					BufferChannelOffset += Texture.ChannelCount;
+				}
+
+				checkf(BufferChannelOffset == NumOfCombinedFeatureChannels, TEXT("Number of channels used by feature does not match the channel count in the buffer."));
+			}
+
+			// Solve the weighted LSR.
+			{
+				FWeightedLSRDesc WeightedLSRDesc;
+				{
+					WeightedLSRDesc.NumOfFeatureChannels = NumOfCombinedFeatureChannels;
+					WeightedLSRDesc.NumOfFeatureChannelsPerFrame = NumOfCombinedFeatureChannels / NumOfRadiances;
+					WeightedLSRDesc.NumOfWeightsPerPixel = NumberOfWeightsPerPixel * NumOfRadiances;
+					WeightedLSRDesc.NumOfWeightsPerPixelPerFrame = NumberOfWeightsPerPixel;
+					WeightedLSRDesc.NumOfRadianceChannels = NumOfCombinedRadianceChannels;
+					WeightedLSRDesc.NumOfRadianceChannelsPerFrame = NumOfCombinedRadianceChannels / NumOfRadiances;
+
+					WeightedLSRDesc.Width = TileSize.X;
+					WeightedLSRDesc.Height = TileSize.Y;
+					WeightedLSRDesc.Offset = PaddingTileOffset;
+					WeightedLSRDesc.TileStartPosition = TileStartPoint;
+					WeightedLSRDesc.NumOfFrames = NumOfRadiances;
+					WeightedLSRDesc.TextureSize = RadianceTileTexture->Desc.Extent;
+					WeightedLSRDesc.SolverType = EWeightedLSRSolverType::Tiled;
+				}
+
+				const int SourceAlbedoFeatureIndex = (NumOfFeatures / NumOfRadiances) * SourceIndex;
+
+				SolveWeightedLSR(
+					GraphBuilder,
+					View,
+					CombinedFeatures,
+					RadianceTileTexture,
+					NonLocalMeanWeightsBuffer,
+					DenoisedTileTexture,
+					WeightedLSRDesc,
+					CombinedRadiances,
+					Features[SourceAlbedoFeatureIndex].Data.Image);
+			}
+
+			// Copy back and accumulate.
+			AddAccumulateTextureRegionPass(GraphBuilder, DenoisedTileTexture, FilteredRadiance, FIntPoint::ZeroValue, PaddedTileRegion.Min, PaddedTileSize);
+
+			if (IsTileDebugEnabled())
+			{
+				break;
+			}
+		}
+
+		// Normalize the image by weights stored in alpha channel.
+		AddNormalizeTexturePass(GraphBuilder, FilteredRadiance);
+
+		// TODO: denoise the alpha channel of the original radiance.
+		// Pass through alpha channel from the source index.
+		AddCopyMirroredTexturePass(GraphBuilder, Radiances[SourceIndex].Data.Image, FilteredRadiance, FIntPoint::ZeroValue, FIntPoint::ZeroValue, FIntPoint::ZeroValue,true/*bAlphaOnly*/);
+
+		return FilteredRadiance;
+	}
+
+	//--------------------------------------------------------------------------------------------------------------------
+	// Bandwidth selection
+
+	FNFORTextureDesc MSEEstimation(
+		FRDGBuilder& GraphBuilder,
+		const FSceneView& View,
+		const FRadianceDesc& Radiance,
+		const FRDGTextureRef FilteredImage)
+	{
+		FIntPoint TextureSize = Radiance.Data.Image->Desc.Extent;
+		FRDGTextureDesc Desc = Radiance.Variance.Image->Desc;
+		Desc.Format = PF_R32_FLOAT;
+		FRDGTextureRef MSE = GraphBuilder.CreateTexture(Desc, TEXT("NFOR.MSE"));
+		NFORDenoise::FNFORTextureDesc NFORMSETexture(MSE, 0, 1, 1);
+
+		typedef FMSEEstimationCS SHADER;
+		SHADER::FParameters* PassParameters = GraphBuilder.AllocParameters<SHADER::FParameters>();
+		{
+			PassParameters->Variance = GraphBuilder.CreateSRV(FRDGTextureSRVDesc::Create(Radiance.Variance.Image));
+			PassParameters->Image = GraphBuilder.CreateSRV(FRDGTextureSRVDesc::Create(Radiance.Data.Image));
+			PassParameters->FilteredImage = GraphBuilder.CreateSRV(FRDGTextureSRVDesc::Create(FilteredImage));
+			PassParameters->TextureSize = TextureSize;
+			PassParameters->VarianceChannelOffset = Radiance.Variance.ChannelOffset;
+			PassParameters->MSE = GraphBuilder.CreateUAV(NFORMSETexture.Image);
+		}
+
+		SHADER::FPermutationDomain ComputeShaderPermutationVector;
+		ComputeShaderPermutationVector.Set<SHADER::FDimensionVarianceType>(Radiance.VarianceType);
+		FGlobalShaderMap* GlobalShaderMap = GetGlobalShaderMap(GMaxRHIFeatureLevel);
+		TShaderMapRef<SHADER> ComputeShader(GlobalShaderMap, ComputeShaderPermutationVector);
+
+		FComputeShaderUtils::AddPass(
+			GraphBuilder,
+			RDG_EVENT_NAME("NFOR::MSEEstimation (Dim=%d,%d)", TextureSize.X, TextureSize.Y),
+			ERDGPassFlags::Compute,
+			ComputeShader,
+			PassParameters,
+			FComputeShaderUtils::GetGroupCount(TextureSize, NON_LOCAL_MEAN_THREAD_GROUP_SIZE));
+
+		return NFORMSETexture;
+	}
+
+	//--------------------------------------------------------------------------------------------------------------------
+	// NFOR filtering and denoising
+
+	void FilterFeatures(
+	FRDGBuilder& GraphBuilder,
+	const FSceneView& View,
+	const TArray<FFeatureDesc>& FeatureDescs)
+	{
+		
+		FNonLocalMeanParameters FeatureNonLocalMeanParameters = GetFeatureNonLocalMeanParameters(0.5f);
+
+		const int NumOfFeatures = FeatureDescs.Num();
+		typedef FNonLocalMeanFilteringCS SHADER;
+
+		int32 BufferChannelOffset = 0;
+		for (int i = 0; i < NumOfFeatures; ++i)
+		{
+			if (FeatureDescs[i].Feature.Image == nullptr 
+				|| FeatureDescs[i].bCleanFeature)
+			{
+				continue;
+			}
+
+			const FFeatureDesc& FeatureDesc = FeatureDescs[i];
+			FRDGTextureRef Feature = FeatureDesc.Feature.Image;
+			FRDGTextureRef FilterdFeature = GraphBuilder.CreateTexture(Feature->Desc, TEXT("NFOR.FilteredFeature"));
+
+			ApplyNonLocalMeanFilter(
+				GraphBuilder,
+				FeatureNonLocalMeanParameters,
+				FeatureDesc.Feature,
+				FeatureDesc.Variance,
+				FeatureDesc.VarianceType,
+				FilterdFeature);
+
+			AddCopyTexturePass(GraphBuilder, FilterdFeature, Feature);
+		}
+	}
+
+	bool FilterMain(
+		FRDGBuilder& GraphBuilder,
+		const FSceneView& View,
+		const TArray<FRadianceDesc>& Radiances,
+		const TArray<FFeatureDesc>& FeatureDescs,
+		const FRDGTextureRef& DenoisedRadiance
+	)
+	{
+		// Frame 0, 1,...,n-1
+		//         0: new frames with feature frame denoised with NLM.
+		// 1,...,n-1: old frames with feature frame denoised with NLM.
+
+		// Denoise for frame n/2. E.g., 
+		//		when n=3, n_m=1,2nd frame is the current frame to denoise. 0, |1|, 2
+		//		when n=5, n_m=2,3rd frame is the current frame to denoise. 0, 1, |2|, 3, 4
+		// Special case when n_a < n
+		//		n_a <= n/2: n_m=n_a
+		//      n_a >  n/2: n_m=n/2
+		// Since the frame size can be very large, we denoise tile by tile and resolve at last
+		// 
+		// Pseudo code:
+		// 
+		//  Preprocessing
+		// 
+		//	For each bandwidth:
+		//		For each tile in tiles:
+		//			collaborative regression(tile)
+		//		denoised = recombine(tiles)
+		// 
+		//  Bandwidth Selection
+  
+		const int32 NumOfTemporalFrames = Radiances.Num();
+		const int32 NumOfFeatures = FeatureDescs.Num();
+		const int32 NumOfFeaturesPerFrame = NumOfFeatures / NumOfTemporalFrames;
+		const int32 SourceRadianceIndex = GetDenoisingFrameIndex(NumOfTemporalFrames);
+		
+		// Preprocessing
+		// Radiance normalization and filtering frames
+		{
+			if (IsPreAlbedoDivideEnabled())
+			{
+				// Latest frame only 
+				for (int i = 0; i < 1; ++i)
+				{
+					FRDGTextureRef AlbedoTex = FeatureDescs[i * NumOfFeaturesPerFrame + 0].Feature.Image; //TODO: unify the index
+					FRDGTextureRef NormalTex = FeatureDescs[i * NumOfFeaturesPerFrame + 1].Feature.Image; //TODO: unify the index
+					FRDGTextureRef NormalVarianceTex = FeatureDescs[i * NumOfFeaturesPerFrame + 1].Variance.Image;
+					FRDGTextureRef MaskTexture = NFORDenoise::GetPreAlbedoDivideMask(GraphBuilder, View, NormalTex, NormalVarianceTex);
+
+					FLinearColor RGBOffset = NFORDenoise::GetPreAlbedoDivideAlbedoOffset();
+
+					NFORDenoise::AddAccumulateConstantRegionPass(GraphBuilder, RGBOffset, AlbedoTex, MaskTexture);
+
+					FRDGTextureRef RadianceTexture = Radiances[i].Data.Image;
+					FRDGTextureRef RadianceVarianceTexture = Radiances[i].Variance.Image;
+
+					// Normalization should apply to both texture to variance.
+					NFORDenoise::AddDivideTextureRegionPass(GraphBuilder, AlbedoTex, RadianceTexture);
+					NFORDenoise::AddNormalizeRadianceVariancePass(GraphBuilder, AlbedoTex, RadianceVarianceTexture);
+				}
+			}
+
+			TArray<FFeatureDesc> LatestFrameFeature = TArray<FFeatureDesc>(FeatureDescs.GetData(), NumOfFeaturesPerFrame);
+			NFORDenoise::FilterFeatures(GraphBuilder, View, LatestFrameFeature);
+		}
+
+		{	// Early out if radiance denoising is not required.
+			EDenoiseFrameCountCondition Condition = GetFrameCountCoundition();
+
+			if (Condition == EDenoiseFrameCountCondition::Equal && 
+				SourceRadianceIndex == INDEX_NONE)
+			{
+				return false;
+			}
+		}
+
+		TArray<FRDGTextureRef> FilteredImages = {};
+		TArray<FNFORTextureDesc> FilteredMSEs = {};
+		TArray<float> Bandwidths = GetBandwidthsConfiguration();
+
+		const bool bPerformBandwidthSelection = IsBandwidthSelectionEnabled() && Bandwidths.Num() == 2;
+		const FRadianceDesc& SourceRadiance = Radiances[SourceRadianceIndex];
+		const int32 RadiancePatchSize = GetNonLocalMeanRadiancePatchSize();
+		const int32 RadiancePatchDistance = GetNonLocalMeanRadiancePatchDistance();
+
+		for (int i = 0; i < Bandwidths.Num(); ++i)
+		{
+			// Collaborative regression.
+			FNonLocalMeanParameters RadianceNonLocalMeanParameters = GetNonLocalMeanParameters(RadiancePatchSize, RadiancePatchDistance, Bandwidths[i]);
+			FRDGTextureRef FilteredImage = CollaborativeRegression(GraphBuilder, View, Radiances, FeatureDescs, RadianceNonLocalMeanParameters);
+
+			FilteredImages.Add(FilteredImage);
+
+			// MSE estimation and filtering.
+			if (bPerformBandwidthSelection)
+			{		
+				FNFORTextureDesc MSE = MSEEstimation(GraphBuilder, View, SourceRadiance, FilteredImage);
+
+				// NLM filtering of MSE texture.
+				FRDGTextureRef FilteredMSETexure = GraphBuilder.CreateTexture(MSE.Image->Desc, TEXT("NFOR.FilteredMSE"));
+				FNonLocalMeanParameters MSENonLocalMeanParameters = GetNonLocalMeanParameters(1, RadiancePatchDistance, 1.0f);
+
+				ApplyNonLocalMeanFilter(
+					GraphBuilder,
+					MSENonLocalMeanParameters,
+					MSE,
+					SourceRadiance.Variance,
+					SourceRadiance.VarianceType,
+					FilteredMSETexure);
+
+				FNFORTextureDesc FilterdMSE = MSE;
+				FilterdMSE.Image = FilteredMSETexure;
+
+				FilteredMSEs.Add(FilterdMSE);
+			}
+		}
+		
+		if (bPerformBandwidthSelection)
+		{
+			FGlobalShaderMap* GlobalShaderMap = GetGlobalShaderMap(GMaxRHIFeatureLevel);
+			FRDGTextureDesc Desc = FilteredMSEs[0].Image->Desc;
+			FNFORTextureDesc NFORSelectionMap = FilteredMSEs[0];
+			NFORSelectionMap.Image = GraphBuilder.CreateTexture(Desc, TEXT("NFOR.SelectionMap"));
+			FIntPoint TextureSize = Desc.Extent;
+			{
+				typedef FGenerateSelectionMapCS SHADER;
+				SHADER::FParameters* PassParameters = GraphBuilder.AllocParameters<SHADER::FParameters>();
+				{
+					PassParameters->FilteredMSEs[0] = GraphBuilder.CreateSRV(FRDGTextureSRVDesc::Create(FilteredMSEs[0].Image));
+					PassParameters->FilteredMSEs[1] = GraphBuilder.CreateSRV(FRDGTextureSRVDesc::Create(FilteredMSEs[1].Image));
+					PassParameters->TextureSize = TextureSize;
+					PassParameters->RWSelectionMap = GraphBuilder.CreateUAV(NFORSelectionMap.Image);
+				}
+
+				TShaderMapRef<SHADER> ComputeShader(GlobalShaderMap);
+
+				FComputeShaderUtils::AddPass(
+					GraphBuilder,
+					RDG_EVENT_NAME("NFOR::GenerateSelectionMap (Dim=%d,%d)", TextureSize.X, TextureSize.Y),
+					ERDGPassFlags::Compute,
+					ComputeShader,
+					PassParameters,
+					FComputeShaderUtils::GetGroupCount(TextureSize, NON_LOCAL_MEAN_THREAD_GROUP_SIZE));
+			}
+
+			// Filter the selection map with image variance
+			FRDGTextureRef FilteredSelectionMap = GraphBuilder.CreateTexture(Desc, TEXT("NFOR.FilteredSelectionMap"));
+			{
+				FNonLocalMeanParameters SelectionMapNonLocalMeanParameters = GetNonLocalMeanParameters(1, RadiancePatchDistance, 1.0f);
+
+				ApplyNonLocalMeanFilter(
+					GraphBuilder,
+					SelectionMapNonLocalMeanParameters,
+					NFORSelectionMap,
+					SourceRadiance.Variance,
+					SourceRadiance.VarianceType,
+					FilteredSelectionMap);
+
+				// Combine the filtered images.
+				{
+					typedef FCombineFilteredImageCS SHADER;
+					SHADER::FParameters* PassParameters = GraphBuilder.AllocParameters<SHADER::FParameters>();
+					{
+						PassParameters->FilteredImages[0] = GraphBuilder.CreateSRV(FRDGTextureSRVDesc::Create(FilteredImages[0]));
+						PassParameters->FilteredImages[1] = GraphBuilder.CreateSRV(FRDGTextureSRVDesc::Create(FilteredImages[1]));
+						PassParameters->SelectionMap = GraphBuilder.CreateSRV(FRDGTextureSRVDesc::Create(FilteredSelectionMap));
+						PassParameters->TextureSize = TextureSize;
+						PassParameters->RWFilteredImage = GraphBuilder.CreateUAV(DenoisedRadiance);
+					}
+
+					TShaderMapRef<SHADER> ComputeShader(GlobalShaderMap);
+
+					FComputeShaderUtils::AddPass(
+						GraphBuilder,
+						RDG_EVENT_NAME("NFOR::ApplySelectionMap (Dim=%d,%d)", TextureSize.X, TextureSize.Y),
+						ERDGPassFlags::Compute,
+						ComputeShader,
+						PassParameters,
+						FComputeShaderUtils::GetGroupCount(TextureSize, NON_LOCAL_MEAN_THREAD_GROUP_SIZE));
+				}
+			}
+
+			// Second regression pass is ignored as there is only one buffer.
+		}
+		else
+		{
+			AddCopyTexturePass(GraphBuilder, FilteredImages[0], DenoisedRadiance);
+		}
+
+		return true;
+	}
+}
\ No newline at end of file
diff --git a/Engine/Plugins/Experimental/NFORDenoise/Source/NFORDenoise/Private/NFORDenoiseCS.h b/Engine/Plugins/Experimental/NFORDenoise/Source/NFORDenoise/Private/NFORDenoiseCS.h
new file mode 100644
--- /dev/null
+++ b/Engine/Plugins/Experimental/NFORDenoise/Source/NFORDenoise/Private/NFORDenoiseCS.h
@@ -0,0 +1,710 @@
+// Copyright Epic Games, Inc. All Rights Reserved.
+
+#pragma once
+
+#include "GlobalShader.h"
+#include "ShaderParameterUtils.h"
+#include "RenderGraphUtils.h"
+#include "ScreenPass.h"
+#include "ShaderCompilerCore.h"
+
+#define NON_LOCAL_MEAN_THREAD_GROUP_SIZE 8
+
+#define TEXTURE_OPS_MULTIPLY			0
+#define TEXTURE_OPS_DIVIDE				1
+#define TEXTURE_OPS_ADD_CONSTANT		2
+#define TEXTURE_OPS_ACCUMULATE			3
+
+namespace NFORDenoise
+{
+
+	int32 GetFrameCount();
+	int32 GetDenoisingFrameIndex(int32 NumberOfFrameInBuffer = -1);
+
+	enum class EVarianceType : uint32
+	{
+		Normal,		// length of vector
+		GreyScale,	// grey scale
+		Colored,	// Not supported at this moment. 
+		MAX
+	};
+
+	enum class EImageChannelCount : uint32
+	{
+		One,
+		Two,
+		Three,
+		Four,
+		MAX
+	};
+
+	enum class EAlbedoDivideRecoverPhase : uint32
+	{
+		Disabled,
+		Each,
+		Final,
+		MAX
+	};
+	
+
+	//--------------------------------------------------------------------------------------------------------------------
+	// General texture operations including: multiply, divide, accumulate, and copy.
+
+	class FTextureMultiplyCS : public FGlobalShader
+	{
+		DECLARE_GLOBAL_SHADER(FTextureMultiplyCS)
+		SHADER_USE_PARAMETER_STRUCT(FTextureMultiplyCS, FGlobalShader)
+	public:
+		BEGIN_SHADER_PARAMETER_STRUCT(FParameters, )
+			SHADER_PARAMETER_RDG_TEXTURE_SRV(Texture2D, Source)
+			SHADER_PARAMETER_RDG_TEXTURE_UAV(RWTexture2D, RWTarget)
+			SHADER_PARAMETER(FIntPoint, SourcePosition)
+			SHADER_PARAMETER(FIntPoint, TargetPosition)
+			SHADER_PARAMETER(FIntPoint, Size)
+			SHADER_PARAMETER(int32, ForceOperation)
+		END_SHADER_PARAMETER_STRUCT()
+
+		static void ModifyCompilationEnvironment(const FGlobalShaderPermutationParameters& InParameters, FShaderCompilerEnvironment& OutEnvironment)
+		{
+			FGlobalShader::ModifyCompilationEnvironment(InParameters, OutEnvironment);
+			OutEnvironment.SetDefine(TEXT("THREAD_GROUP_SIZE"), NON_LOCAL_MEAN_THREAD_GROUP_SIZE);
+			OutEnvironment.SetDefine(TEXT("TEXTURE_OPS"), TEXTURE_OPS_MULTIPLY);
+			OutEnvironment.CompilerFlags.Add(CFLAG_AllowTypedUAVLoads);
+		}
+	};
+
+	class FTextureDivideCS : public FGlobalShader
+	{
+		DECLARE_GLOBAL_SHADER(FTextureDivideCS)
+		SHADER_USE_PARAMETER_STRUCT(FTextureDivideCS, FGlobalShader)
+	public:
+		BEGIN_SHADER_PARAMETER_STRUCT(FParameters, )
+			SHADER_PARAMETER_RDG_TEXTURE_SRV(Texture2D, Source)
+			SHADER_PARAMETER_RDG_TEXTURE_UAV(RWTexture2D, RWTarget)
+			SHADER_PARAMETER(FIntPoint, SourcePosition)
+			SHADER_PARAMETER(FIntPoint, TargetPosition)
+			SHADER_PARAMETER(FIntPoint, Size)
+			SHADER_PARAMETER(int32, ForceOperation)
+		END_SHADER_PARAMETER_STRUCT()
+
+		static void ModifyCompilationEnvironment(const FGlobalShaderPermutationParameters& InParameters, FShaderCompilerEnvironment& OutEnvironment)
+		{
+			FGlobalShader::ModifyCompilationEnvironment(InParameters, OutEnvironment);
+			OutEnvironment.SetDefine(TEXT("THREAD_GROUP_SIZE"), NON_LOCAL_MEAN_THREAD_GROUP_SIZE);
+			OutEnvironment.SetDefine(TEXT("TEXTURE_OPS"), TEXTURE_OPS_DIVIDE);
+			OutEnvironment.CompilerFlags.Add(CFLAG_AllowTypedUAVLoads);
+		}
+	};
+
+	class FTextureAccumulateConstantCS : public FGlobalShader
+	{
+		DECLARE_GLOBAL_SHADER(FTextureAccumulateConstantCS)
+		SHADER_USE_PARAMETER_STRUCT(FTextureAccumulateConstantCS, FGlobalShader)
+	public:
+		BEGIN_SHADER_PARAMETER_STRUCT(FParameters, )
+			SHADER_PARAMETER_RDG_TEXTURE_UAV(RWTexture2D, RWTarget)
+			SHADER_PARAMETER_RDG_TEXTURE_SRV(Texture2D, Mask)
+			SHADER_PARAMETER(FIntPoint, TargetPosition)
+			SHADER_PARAMETER(FIntPoint, Size)
+			SHADER_PARAMETER(FLinearColor, ConstantValue)
+		END_SHADER_PARAMETER_STRUCT()
+
+		static void ModifyCompilationEnvironment(const FGlobalShaderPermutationParameters& InParameters, FShaderCompilerEnvironment& OutEnvironment)
+		{
+			FGlobalShader::ModifyCompilationEnvironment(InParameters, OutEnvironment);
+			OutEnvironment.SetDefine(TEXT("THREAD_GROUP_SIZE"), NON_LOCAL_MEAN_THREAD_GROUP_SIZE);
+			OutEnvironment.SetDefine(TEXT("TEXTURE_OPS"), TEXTURE_OPS_ADD_CONSTANT);
+			OutEnvironment.CompilerFlags.Add(CFLAG_AllowTypedUAVLoads);
+		}
+
+		class FDimensionAccumulateByMask : SHADER_PERMUTATION_BOOL("ACCUMULATE_BY_MASK");
+		using FPermutationDomain = TShaderPermutationDomain<FDimensionAccumulateByMask>;
+	};
+
+	// Accumulate all channels of a texture onto another texture for a given region.
+	class FTextureAccumulateCS : public FGlobalShader
+	{
+		DECLARE_GLOBAL_SHADER(FTextureAccumulateCS)
+		SHADER_USE_PARAMETER_STRUCT(FTextureAccumulateCS, FGlobalShader)
+	public:
+		BEGIN_SHADER_PARAMETER_STRUCT(FParameters, )
+			SHADER_PARAMETER_RDG_TEXTURE_SRV(Texture2D, Source)
+			SHADER_PARAMETER_RDG_TEXTURE_UAV(RWTexture2D, RWTarget)
+			SHADER_PARAMETER(FIntPoint, SourcePosition)
+			SHADER_PARAMETER(FIntPoint, TargetPosition)
+			SHADER_PARAMETER(FIntPoint, Size)
+		END_SHADER_PARAMETER_STRUCT()
+
+		static void ModifyCompilationEnvironment(const FGlobalShaderPermutationParameters& InParameters, FShaderCompilerEnvironment& OutEnvironment)
+		{
+			FGlobalShader::ModifyCompilationEnvironment(InParameters, OutEnvironment);
+			OutEnvironment.SetDefine(TEXT("THREAD_GROUP_SIZE"), NON_LOCAL_MEAN_THREAD_GROUP_SIZE);
+			OutEnvironment.SetDefine(TEXT("TEXTURE_OPS"), TEXTURE_OPS_ACCUMULATE);
+			OutEnvironment.CompilerFlags.Add(CFLAG_AllowTypedUAVLoads);
+		}
+	};
+
+	class FCopyTexturePS : public FGlobalShader
+	{
+		DECLARE_GLOBAL_SHADER(FCopyTexturePS)
+		SHADER_USE_PARAMETER_STRUCT(FCopyTexturePS, FGlobalShader)
+	public:
+		BEGIN_SHADER_PARAMETER_STRUCT(FParameters, )
+			SHADER_PARAMETER_RDG_TEXTURE_SRV(Texture2D, Source)
+			SHADER_PARAMETER(FIntPoint, SourceOffset)
+			SHADER_PARAMETER(FIntPoint, TextureSize)
+			RENDER_TARGET_BINDING_SLOTS()
+		END_SHADER_PARAMETER_STRUCT()
+	};
+
+	// TargetTexture.rgb = TargetTexture.rgb * lerp (1.0f,SourceTexture.rgb, (SourceTexture.rgb != 0 || bForceMultiply));
+	void AddMultiplyTextureRegionPass(FRDGBuilder& GraphBuilder, const FRDGTextureRef& SourceTexture, const FRDGTextureRef& TargetTexture,
+		bool bForceMultiply = true, FIntPoint SourcePosition = FIntPoint::ZeroValue, FIntPoint TargetPosition = FIntPoint::ZeroValue, FIntPoint Size = FIntPoint::ZeroValue);
+
+	// TargetTexture.rgb = TargetTexture.rgb / lerp (1.0f,SourceTexture.rgb, SourceTexture.rgb != 0 || bForceDivide);
+	void AddDivideTextureRegionPass(FRDGBuilder& GraphBuilder, const FRDGTextureRef& SourceTexture, const FRDGTextureRef& TargetTexture,
+		bool bForceDivide = true, FIntPoint SourcePosition = FIntPoint::ZeroValue, FIntPoint TargetPosition = FIntPoint::ZeroValue, FIntPoint Size = FIntPoint::ZeroValue);
+	
+	// TargetTexture += SourceTexture;
+	void AddAccumulateTextureRegionPass(FRDGBuilder& GraphBuilder, const FRDGTextureRef& SourceTexture, const FRDGTextureRef& TargetTexture,
+		FIntPoint SourcePosition = FIntPoint::ZeroValue, FIntPoint TargetPosition = FIntPoint::ZeroValue, FIntPoint Size = FIntPoint::ZeroValue);
+
+	// TargetTexture.rgb = TargetTexture.rgb + lerp(ConstantValue, ConstantValue[mask], Mask != nullptr);
+	void AddAccumulateConstantRegionPass(FRDGBuilder& GraphBuilder, const FLinearColor& ConstantValue, const FRDGTextureRef& TargetTexture, 
+		const FRDGTextureRef& Mask = nullptr, FIntPoint SourcePosition = FIntPoint::ZeroValue, FIntPoint TargetPosition = FIntPoint::ZeroValue, FIntPoint Size = FIntPoint::ZeroValue);
+
+	// Copy texture with mirrored border of source texture;
+	void AddCopyMirroredTexturePass(FRDGBuilder& GraphBuilder, const FRDGTextureRef& SourceTexture, const FRDGTextureRef& TargetTexture,
+		FIntPoint SourcePosition = FIntPoint::ZeroValue, FIntPoint TargetPosition = FIntPoint::ZeroValue, FIntPoint Size = FIntPoint::ZeroValue, bool bAlphaOnly = false);
+	
+
+	//--------------------------------------------------------------------------------------------------------------------
+	// Radiance normalization by albedo
+	
+	class FClassifyPreAlbedoDivideMaskIdCS : public FGlobalShader
+	{
+		DECLARE_GLOBAL_SHADER(FClassifyPreAlbedoDivideMaskIdCS);
+		SHADER_USE_PARAMETER_STRUCT(FClassifyPreAlbedoDivideMaskIdCS, FGlobalShader);
+	public:
+		BEGIN_SHADER_PARAMETER_STRUCT(FParameters, )
+			SHADER_PARAMETER_RDG_TEXTURE_SRV(Texture2D, Normal)
+			SHADER_PARAMETER_RDG_TEXTURE_SRV(Texture2D, NormalVariance)
+			SHADER_PARAMETER_RDG_TEXTURE_UAV(RWTexture2D, RWMask)
+
+			SHADER_PARAMETER(FIntPoint, TextureSize)
+			END_SHADER_PARAMETER_STRUCT()
+
+			static void ModifyCompilationEnvironment(const FGlobalShaderPermutationParameters& InParameters, FShaderCompilerEnvironment& OutEnvironment)
+		{
+			FGlobalShader::ModifyCompilationEnvironment(InParameters, OutEnvironment);
+			OutEnvironment.SetDefine(TEXT("THREAD_GROUP_SIZE"), NON_LOCAL_MEAN_THREAD_GROUP_SIZE);
+		}
+	};
+
+	// Based on Taylor expansion, sigma_{normalized radiance} \approx = sigma_{radiance} / sigma_{albedo}, if normalized radiance = radiance / albedo
+	class FNormalizeRadianceVarianceByAlbedoCS : public FGlobalShader
+	{
+		DECLARE_GLOBAL_SHADER(FNormalizeRadianceVarianceByAlbedoCS)
+			SHADER_USE_PARAMETER_STRUCT(FNormalizeRadianceVarianceByAlbedoCS, FGlobalShader)
+	public:
+		BEGIN_SHADER_PARAMETER_STRUCT(FParameters, )
+			SHADER_PARAMETER_RDG_TEXTURE_SRV(Texture2D<float4>, Albedo)
+			SHADER_PARAMETER_RDG_TEXTURE_UAV(RWTexture2D<float4>, RWRadianceVariance)
+			SHADER_PARAMETER(FIntPoint, Size)
+			END_SHADER_PARAMETER_STRUCT()
+
+			static void ModifyCompilationEnvironment(const FGlobalShaderPermutationParameters& InParameters, FShaderCompilerEnvironment& OutEnvironment)
+		{
+			FGlobalShader::ModifyCompilationEnvironment(InParameters, OutEnvironment);
+			OutEnvironment.SetDefine(TEXT("THREAD_GROUP_SIZE"), NON_LOCAL_MEAN_THREAD_GROUP_SIZE);
+			OutEnvironment.CompilerFlags.Add(CFLAG_AllowTypedUAVLoads);
+		}
+	};
+
+	FRDGTextureRef GetPreAlbedoDivideMask(FRDGBuilder& GraphBuilder, const FSceneView& View, const FRDGTextureRef& Normal, const FRDGTextureRef& NormalVariance);
+
+	void AddNormalizeRadianceVariancePass(FRDGBuilder& GraphBuilder, const FRDGTextureRef& Albedo, const FRDGTextureRef& RadianceVariance);
+
+
+	//--------------------------------------------------------------------------------------------------------------------
+	// Non-local mean weight and filtering
+	
+	BEGIN_SHADER_PARAMETER_STRUCT(FNonLocalMeanParameters, )
+		SHADER_PARAMETER(int32, PatchSize)
+		SHADER_PARAMETER(int32, PatchDistance)
+		SHADER_PARAMETER(float, Bandwidth)
+	END_SHADER_PARAMETER_STRUCT()
+
+	FNonLocalMeanParameters GetNonLocalMeanParameters(int32 PatchSize, int32 PatchDistance, float Bandwidth);
+
+	// Output the non-local mean filtered image (feature) based on variance
+	// Input: Image, variance, Guide(optional), NonLocalMean parameters
+	//		 Dim: WxHx?, WxHx?
+	// Output: Image holding  prefiltered features
+	//		 Dim: WxHx?
+	class FNonLocalMeanFilteringCS : public FGlobalShader
+	{
+		DECLARE_GLOBAL_SHADER(FNonLocalMeanFilteringCS)
+		SHADER_USE_PARAMETER_STRUCT(FNonLocalMeanFilteringCS, FGlobalShader)
+	public:
+
+		BEGIN_SHADER_PARAMETER_STRUCT(FParameters, )
+			SHADER_PARAMETER_STRUCT_INCLUDE(FNonLocalMeanParameters, NLMParams)
+			SHADER_PARAMETER_RDG_TEXTURE_SRV(Texture2D<float4>, Guide)
+			SHADER_PARAMETER_RDG_TEXTURE_SRV(Texture2D, Variance)
+			SHADER_PARAMETER_RDG_TEXTURE_SRV(Texture2D, Image)
+			SHADER_PARAMETER(FIntPoint, TextureSize)
+			SHADER_PARAMETER(int32, VarianceChannelOffset)
+			SHADER_PARAMETER(int32, DenoisingChannelCount)
+			SHADER_PARAMETER_RDG_TEXTURE_UAV(RWTexture2D, DenoisedImage)
+			END_SHADER_PARAMETER_STRUCT()
+
+		static void ModifyCompilationEnvironment(const FGlobalShaderPermutationParameters& InParameters, FShaderCompilerEnvironment& OutEnvironment)
+		{
+			FGlobalShader::ModifyCompilationEnvironment(InParameters, OutEnvironment);
+			OutEnvironment.SetDefine(TEXT("THREAD_GROUP_SIZE"), NON_LOCAL_MEAN_THREAD_GROUP_SIZE);
+			OutEnvironment.SetDefine(TEXT("NONLOCALMEAN_SEPARATE_SOURCE"), 0);
+		}
+
+		class FDimensionVarianceType : SHADER_PERMUTATION_ENUM_CLASS("IMAGE_VARIANCE_TYPE", EVarianceType);
+		class FDimensionUseGuide : SHADER_PERMUTATION_BOOL("USE_GUIDE"); //TODO
+		class FDimensionImageChannelCount : SHADER_PERMUTATION_RANGE_INT("SOURCE_CHANNEL_COUNT", 1, static_cast<int>(EImageChannelCount::MAX));
+		class FDimPreAlbedoDivide : SHADER_PERMUTATION_ENUM_CLASS("PRE_ALBEDO_DIVIDE", EAlbedoDivideRecoverPhase);
+		using FPermutationDomain = TShaderPermutationDomain<FDimensionVarianceType, FDimensionUseGuide, FDimensionImageChannelCount, FDimPreAlbedoDivide>;
+	};
+
+
+	// Output the non-local mean weights for each pixel, used to solve the weighted least squares problem.
+	// Input: Image, variance, NonLocalMean parameters
+	//		 Dim: WxHx3, WxHx?
+	// Output: Buffer holding  weights for each pixel
+	//		 Dim: WxHx(2*N+1)^2, where N = NLMParams.PatchDistance
+	//		 TODO?: optimize weights sharing as center(i,j) with offset (dx,dy) has
+	//		 w(i,j,dx,dy)=(i+dx,j+dy,-dx,-dy), 
+
+	class FNonLocalMeanWeightsCS : public FGlobalShader
+	{
+		DECLARE_GLOBAL_SHADER(FNonLocalMeanWeightsCS)
+		SHADER_USE_PARAMETER_STRUCT(FNonLocalMeanWeightsCS, FGlobalShader)
+	public:
+
+		BEGIN_SHADER_PARAMETER_STRUCT(FParameters, )
+			SHADER_PARAMETER_STRUCT_INCLUDE(FNonLocalMeanFilteringCS::FParameters, CommonParameters)
+			SHADER_PARAMETER(FIntRect, Region)
+			SHADER_PARAMETER_RDG_TEXTURE_SRV(Texture2D, TargetImage)
+			SHADER_PARAMETER_RDG_TEXTURE_SRV(Texture2D, TargetVariance)
+			SHADER_PARAMETER_RDG_BUFFER_UAV(RWBuffer<float>, NonLocalMeanWeights)
+		END_SHADER_PARAMETER_STRUCT()
+
+		static void ModifyCompilationEnvironment(const FGlobalShaderPermutationParameters& InParameters, FShaderCompilerEnvironment& OutEnvironment)
+		{
+			FGlobalShader::ModifyCompilationEnvironment(InParameters, OutEnvironment);
+			OutEnvironment.SetDefine(TEXT("THREAD_GROUP_SIZE"), NON_LOCAL_MEAN_THREAD_GROUP_SIZE);
+		}
+
+		class FDimensionVarianceType : SHADER_PERMUTATION_ENUM_CLASS("IMAGE_VARIANCE_TYPE", EVarianceType);
+		class FDimensionUseGuide : SHADER_PERMUTATION_BOOL("USE_GUIDE");
+		class FDimensionImageChannelCount : SHADER_PERMUTATION_RANGE_INT("SOURCE_CHANNEL_COUNT", 1, static_cast<int>(EImageChannelCount::MAX));
+		class FDimensionSeparateSourceTarget : SHADER_PERMUTATION_BOOL("NONLOCALMEAN_SEPARATE_SOURCE");
+		class FDimPreAlbedoDivide : SHADER_PERMUTATION_ENUM_CLASS("PRE_ALBEDO_DIVIDE", EAlbedoDivideRecoverPhase);
+		using FPermutationDomain = TShaderPermutationDomain<FDimensionVarianceType, FDimensionUseGuide, FDimensionImageChannelCount, FDimensionSeparateSourceTarget, FDimPreAlbedoDivide>;
+	};
+
+
+	//--------------------------------------------------------------------------------------------------------------------
+	// Collaborative filtering shaders
+	//	1. Tiling
+	// 
+	// Copy textures to buffer based on copy channel config
+	// Buffer layout
+	// |Pixel 1                          | Pixel 2|...|Pixel n|
+	//  tex1.rgb|tex2.rgb|tex3.rgb
+	class FCopyTextureToBufferCS : public FGlobalShader
+	{
+		DECLARE_GLOBAL_SHADER(FCopyTextureToBufferCS)
+		SHADER_USE_PARAMETER_STRUCT(FCopyTextureToBufferCS, FGlobalShader)
+	public:
+
+		BEGIN_SHADER_PARAMETER_STRUCT(FParameters, )
+			SHADER_PARAMETER_RDG_TEXTURE_SRV(Texture2D, Source)
+			SHADER_PARAMETER_RDG_BUFFER_UAV(RWBuffer<float>, Dest)
+			SHADER_PARAMETER(FIntPoint, TextureSize)
+			SHADER_PARAMETER(int32, CopyChannelCount)
+			SHADER_PARAMETER(int32, CopyChannelOffset)
+			SHADER_PARAMETER(int32, BufferChannelOffset)
+			SHADER_PARAMETER(int32, BufferChannelSize)
+			SHADER_PARAMETER(FIntRect, CopyRegion)
+		END_SHADER_PARAMETER_STRUCT()
+
+		static void ModifyCompilationEnvironment(const FGlobalShaderPermutationParameters& InParameters, FShaderCompilerEnvironment& OutEnvironment)
+		{
+			FGlobalShader::ModifyCompilationEnvironment(InParameters, OutEnvironment);
+			OutEnvironment.SetDefine(TEXT("THREAD_GROUP_SIZE"), NON_LOCAL_MEAN_THREAD_GROUP_SIZE);
+		}
+
+		static const int kMaxSourceChannelCount = 4;
+
+		class FDimensionSourceChannelCount : SHADER_PERMUTATION_RANGE_INT("SOURCE_CHANNEL_COUNT", 1, kMaxSourceChannelCount);
+		using FPermutationDomain = TShaderPermutationDomain<FDimensionSourceChannelCount>;
+	};
+
+	class FNormalizeTextureCS : public FGlobalShader
+	{
+		DECLARE_GLOBAL_SHADER(FNormalizeTextureCS)
+		SHADER_USE_PARAMETER_STRUCT(FNormalizeTextureCS, FGlobalShader)
+	public:
+		BEGIN_SHADER_PARAMETER_STRUCT(FParameters, )
+			SHADER_PARAMETER_RDG_TEXTURE_UAV(RWTexture2D, RWSource)
+			SHADER_PARAMETER(FIntPoint, TextureSize)
+		END_SHADER_PARAMETER_STRUCT()
+
+			static void ModifyCompilationEnvironment(const FGlobalShaderPermutationParameters& InParameters, FShaderCompilerEnvironment& OutEnvironment)
+		{
+			FGlobalShader::ModifyCompilationEnvironment(InParameters, OutEnvironment);
+			OutEnvironment.SetDefine(TEXT("THREAD_GROUP_SIZE"), NON_LOCAL_MEAN_THREAD_GROUP_SIZE);
+			OutEnvironment.CompilerFlags.Add(CFLAG_AllowTypedUAVLoads);
+		}
+	};
+
+
+	//==============================================================
+	// Regression utility shaders
+	// Given, X, Y, W buffer, calculate
+	// A = X^T W X
+	// B = X^T W Y
+	// for each pixel footage, where 
+	//		A has a dimension of |F|x|F|
+	//		B has a dimension of |F|x|3|, |F| is the dimension of the feature count
+	// typically it would be, 7 (1 constant, 3 albedo, 3 normal), TODO: add normalized depth.
+	// 1. Calculate weighted quadratic form A, and the generalized weighted multiplication B
+	// 2. Use solver to get reconstruction weight A^{-1} B (Ax = B, A=7x7, B=7x3)
+	// 3. A1 Multiplication of X B and write to pixels for each image. 
+	//	  A2. Gather temporal pixels based on weights
+	// or
+	// 3. B1 Multiplication of X B and gather at the same time.
+	namespace RegressionKernel
+	{
+		FRDGBufferRef AllocateMatrixfBuffer(FRDGBuilder& GraphBuilder, int32 NumOfMatrices, int32 Dim0, int32 Dim1, const TCHAR* Name);
+
+		enum class EWeightedMultiplicationType : uint8
+		{
+			Quadratic,		// X^T W X
+			Generalized,	// X^T W Y
+			MAX
+		};
+
+		static const TCHAR* GetEventName(EWeightedMultiplicationType WeightedMultiplicationType)
+		{
+			static const TCHAR* const kEventNames[] = {
+				TEXT("X^TWX"),
+				TEXT("X^TWY"),
+			};
+			static_assert(UE_ARRAY_COUNT(kEventNames) == int32(EWeightedMultiplicationType::MAX), "Fix me");
+			return kEventNames[int32(WeightedMultiplicationType)];
+		}
+
+		// Calculate
+		// X^T Diag(W) Y, where Y might be equal to X
+		// Note that the data is gathered in place to calculate the matrix
+		// for each pixel. This avoids data duplication saved to memory at the cost of increasing memory bandwidth utilization
+		// it is useful when memory size is small.
+		class FInPlaceBatchedMatrixMultiplicationCS : public FGlobalShader
+		{
+			DECLARE_GLOBAL_SHADER(FInPlaceBatchedMatrixMultiplicationCS);
+			SHADER_USE_PARAMETER_STRUCT(FInPlaceBatchedMatrixMultiplicationCS, FGlobalShader);
+
+		public:
+			static FRDGBufferRef AllocateResultBuffer(FRDGBuilder& GraphBuilder, FIntPoint Size, int F, int A);
+
+			BEGIN_SHADER_PARAMETER_STRUCT(FParameters, )
+				SHADER_PARAMETER_RDG_BUFFER_SRV(Buffer<float>, X)
+				SHADER_PARAMETER_RDG_BUFFER_SRV(Buffer<float>, W)
+				SHADER_PARAMETER_RDG_BUFFER_SRV(Buffer<float>, Y)
+				SHADER_PARAMETER_RDG_BUFFER_UAV(RWBuffer<float>, Result)
+
+				// Parameters for matrix multiplication
+				SHADER_PARAMETER(FIntPoint, XDim) // NxF
+				SHADER_PARAMETER(int32, WDim)	  // N	
+				SHADER_PARAMETER(FIntPoint, YDim) // NxA
+
+				// Parameters for the batched data storage for the image
+				SHADER_PARAMETER(FIntPoint, TextureSize)
+				SHADER_PARAMETER(int32, PatchDistance)
+				SHADER_PARAMETER(int32, NumOfTemporalFrames)
+				SHADER_PARAMETER(int32, NumOfWeigthsPerPixelPerFrame)
+
+				// Used to accelerate the performance while maintaining good quality.
+				SHADER_PARAMETER(int32, SamplingStep)
+
+				SHADER_PARAMETER(int32, SourceFrameIndex) // which frame is currently to be denoised
+
+			END_SHADER_PARAMETER_STRUCT()
+
+
+			static int GetThreadGroupSize()
+			{
+				return NON_LOCAL_MEAN_THREAD_GROUP_SIZE;
+			}
+			static void ModifyCompilationEnvironment(const FGlobalShaderPermutationParameters& InParameters, FShaderCompilerEnvironment& OutEnvironment)
+			{
+				FGlobalShader::ModifyCompilationEnvironment(InParameters, OutEnvironment);
+				OutEnvironment.SetDefine(TEXT("THREAD_GROUP_SIZE"), GetThreadGroupSize());
+			}
+
+			static const int kMaxSourceChannelCount = 4;
+
+			class FDimWeightedMultiplicationType : SHADER_PERMUTATION_ENUM_CLASS("WEIGHTED_MULTIPLICATION_TYPE", EWeightedMultiplicationType);
+			class FDimAddConstantFeatureDim : SHADER_PERMUTATION_BOOL("APPEND_CONSTANT_DIMENSION_TO_X");
+			class FDimNumFeature : SHADER_PERMUTATION_RANGE_INT("NUM_FEATURE", 6, 3);
+			class FDimOptimizeTargetMatrixMultiplication : SHADER_PERMUTATION_BOOL("SMALL_MATRIX_OPTIMIZE");
+			class FDimUseSamplingStep : SHADER_PERMUTATION_BOOL("USE_SAMPLING_STEP");
+			using FPermutationDomain = TShaderPermutationDomain<FDimWeightedMultiplicationType, FDimAddConstantFeatureDim, FDimOptimizeTargetMatrixMultiplication, FDimNumFeature, FDimUseSamplingStep>;
+		};
+
+
+		class FLinearSolverCS : public FGlobalShader
+		{
+			DECLARE_GLOBAL_SHADER(FLinearSolverCS);
+			SHADER_USE_PARAMETER_STRUCT(FLinearSolverCS, FGlobalShader);
+
+		public:
+			BEGIN_SHADER_PARAMETER_STRUCT(FParameters, )
+				SHADER_PARAMETER_RDG_BUFFER_SRV(Buffer<float>, A)
+				SHADER_PARAMETER_RDG_BUFFER_SRV(Buffer<float>, B)
+				SHADER_PARAMETER_RDG_BUFFER_UAV(RWBuffer<float>, Result)
+
+				SHADER_PARAMETER(FIntPoint, ADim) // FxF
+				SHADER_PARAMETER(FIntPoint, BDim) // FxA
+
+				SHADER_PARAMETER(int, NumOfElements)
+				SHADER_PARAMETER(int, NumOfElementsPerRow)
+			END_SHADER_PARAMETER_STRUCT()
+
+			static void ModifyCompilationEnvironment(const FGlobalShaderPermutationParameters& InParameters, FShaderCompilerEnvironment& OutEnvironment)
+			{
+				FGlobalShader::ModifyCompilationEnvironment(InParameters, OutEnvironment);
+				OutEnvironment.SetDefine(TEXT("THREAD_GROUP_SIZE"), NON_LOCAL_MEAN_THREAD_GROUP_SIZE);
+			}
+
+			class FDimNumFeature : SHADER_PERMUTATION_RANGE_INT("NUM_FEATURE", 6, 3);
+			using FPermutationDomain = TShaderPermutationDomain<FDimNumFeature>;
+		};
+
+		// reconstruct with the weights into an image
+		// Given X(|N|x|F|) and B (|F|x|3|)
+		// spatial only (temporal frame, T = 1), Result = X * B, accumulate on each pixel to get shaper result
+		// spatial temporal (T =3,5).
+		//	Result_{pi} = X_i * w * B {i=T/2} + \sum_{i!=T/2}\sum_{j \in patch_i} (w * X_j * B)
+
+		class FReconstructSpatialTemporalImage : public FGlobalShader
+		{
+			DECLARE_GLOBAL_SHADER(FReconstructSpatialTemporalImage);
+			SHADER_USE_PARAMETER_STRUCT(FReconstructSpatialTemporalImage, FGlobalShader);
+
+		public:
+			BEGIN_SHADER_PARAMETER_STRUCT(FParameters, )
+				SHADER_PARAMETER_RDG_BUFFER_SRV(Buffer<float>, X)
+				SHADER_PARAMETER_RDG_BUFFER_SRV(Buffer<float>, W)
+				SHADER_PARAMETER_RDG_BUFFER_SRV(Buffer<float>, B)
+				SHADER_PARAMETER_RDG_TEXTURE_UAV(RWTexture2D<float4>, RWReconstruction)
+				SHADER_PARAMETER_RDG_BUFFER_UAV(RWStructuredBuffer<uint4>, RWReconstructBuffer)
+				SHADER_PARAMETER_RDG_TEXTURE_UAV(RWTexture2D<UlongType>, RWReconstructBuffer64)
+
+
+				SHADER_PARAMETER(FIntPoint, XDim) // N(=PxT) x F
+				SHADER_PARAMETER(int32, WDim)	  // N	
+				SHADER_PARAMETER(FIntPoint, BDim) // Px3
+
+				// Parameters for the batched data storage for the image
+				SHADER_PARAMETER(FIntPoint, TextureSize)
+				SHADER_PARAMETER(int32, PatchDistance)
+				SHADER_PARAMETER(int32, FrameIndex)
+
+				SHADER_PARAMETER(int32, NumOfTemporalFrames)
+				SHADER_PARAMETER(int32, NumOfWeigthsPerPixelPerFrame)
+
+				SHADER_PARAMETER(float, AlbedoOffset)
+			END_SHADER_PARAMETER_STRUCT()
+
+
+			static void ModifyCompilationEnvironment(const FGlobalShaderPermutationParameters& InParameters, FShaderCompilerEnvironment& OutEnvironment)
+			{
+				FGlobalShader::ModifyCompilationEnvironment(InParameters, OutEnvironment);
+				OutEnvironment.SetDefine(TEXT("THREAD_GROUP_SIZE"), NON_LOCAL_MEAN_THREAD_GROUP_SIZE);
+				OutEnvironment.CompilerFlags.Add(CFLAG_AllowTypedUAVLoads);
+			}
+
+			enum class EReconstructionType : uint8
+			{
+				Scatter,
+				Gather,
+				MAX
+			};
+
+			static const TCHAR* GetEventName(EReconstructionType ReconstructionType)
+			{
+				static const TCHAR* const kEventNames[] = {
+					TEXT("Scatter"),
+					TEXT("Gather"),
+				};
+				static_assert(UE_ARRAY_COUNT(kEventNames) == int32(EWeightedMultiplicationType::MAX), "Fix me");
+				return kEventNames[int32(ReconstructionType)];
+			}
+
+			class FDimReconstructionType : SHADER_PERMUTATION_ENUM_CLASS("RECONSTRUCTION_TYPE", EReconstructionType);
+			class FDimPreAlbedoDivide : SHADER_PERMUTATION_ENUM_CLASS("PRE_ALBEDO_DIVIDE", EAlbedoDivideRecoverPhase);
+			class FDimNumFeature : SHADER_PERMUTATION_RANGE_INT("NUM_FEATURE", 6, 3);
+			using FPermutationDomain = TShaderPermutationDomain<FDimReconstructionType, FDimPreAlbedoDivide, FDimNumFeature>;
+		};
+	}
+
+	class FAccumulateBufferToTextureCS : public FGlobalShader
+	{
+		DECLARE_GLOBAL_SHADER(FAccumulateBufferToTextureCS);
+		SHADER_USE_PARAMETER_STRUCT(FAccumulateBufferToTextureCS, FGlobalShader);
+
+	public:
+		BEGIN_SHADER_PARAMETER_STRUCT(FParameters, )
+			SHADER_PARAMETER_RDG_BUFFER_SRV(StructuredBuffer<uint4>, StructuredBufferSource)
+			SHADER_PARAMETER_RDG_TEXTURE_SRV(Texture2D<UlongType>, ReconstructBuffer64)
+			SHADER_PARAMETER_RDG_TEXTURE_UAV(RWTexture2D<float4>, RWTarget)
+
+			SHADER_PARAMETER(FIntPoint, TextureSize)
+		END_SHADER_PARAMETER_STRUCT()
+
+		static void ModifyCompilationEnvironment(const FGlobalShaderPermutationParameters& InParameters, FShaderCompilerEnvironment& OutEnvironment)
+		{
+			FGlobalShader::ModifyCompilationEnvironment(InParameters, OutEnvironment);
+			OutEnvironment.SetDefine(TEXT("THREAD_GROUP_SIZE"), NON_LOCAL_MEAN_THREAD_GROUP_SIZE);
+			OutEnvironment.CompilerFlags.Add(CFLAG_AllowTypedUAVLoads);
+		}
+
+		class FDimPreAlbedoDivide : SHADER_PERMUTATION_ENUM_CLASS("PRE_ALBEDO_DIVIDE", EAlbedoDivideRecoverPhase);
+		using FPermutationDomain = TShaderPermutationDomain<FDimPreAlbedoDivide>;
+	};
+
+
+	//--------------------------------------------------------------------------------------------------------------------
+	// Bandwidth selection
+
+	class FMSEEstimationCS : public FGlobalShader
+	{
+		DECLARE_GLOBAL_SHADER(FMSEEstimationCS)
+			SHADER_USE_PARAMETER_STRUCT(FMSEEstimationCS, FGlobalShader)
+
+	public:
+		BEGIN_SHADER_PARAMETER_STRUCT(FParameters, )
+			SHADER_PARAMETER_RDG_TEXTURE_SRV(Texture2D, Variance)
+			SHADER_PARAMETER_RDG_TEXTURE_SRV(Texture2D<float4>, Image)
+			SHADER_PARAMETER_RDG_TEXTURE_SRV(Texture2D<float4>, FilteredImage)
+			SHADER_PARAMETER(FIntPoint, TextureSize)
+			SHADER_PARAMETER(int32, VarianceChannelOffset)
+			SHADER_PARAMETER_RDG_TEXTURE_UAV(RWTexture2D, MSE)
+			END_SHADER_PARAMETER_STRUCT()
+
+			static void ModifyCompilationEnvironment(const FGlobalShaderPermutationParameters& InParameters, FShaderCompilerEnvironment& OutEnvironment)
+		{
+			FGlobalShader::ModifyCompilationEnvironment(InParameters, OutEnvironment);
+			OutEnvironment.SetDefine(TEXT("THREAD_GROUP_SIZE"), NON_LOCAL_MEAN_THREAD_GROUP_SIZE);
+			OutEnvironment.SetDefine(TEXT("SOURCE_CHANNEL_COUNT"), 4);
+		}
+
+		class FDimensionVarianceType : SHADER_PERMUTATION_ENUM_CLASS("IMAGE_VARIANCE_TYPE", EVarianceType);
+		using FPermutationDomain = TShaderPermutationDomain<FDimensionVarianceType>;
+	};
+
+	class FGenerateSelectionMapCS : public FGlobalShader
+	{
+		DECLARE_GLOBAL_SHADER(FGenerateSelectionMapCS)
+			SHADER_USE_PARAMETER_STRUCT(FGenerateSelectionMapCS, FGlobalShader)
+
+	public:
+		BEGIN_SHADER_PARAMETER_STRUCT(FParameters, )
+			SHADER_PARAMETER_RDG_TEXTURE_SRV_ARRAY(Texture2D, FilteredMSEs, [2])
+			SHADER_PARAMETER(FIntPoint, TextureSize)
+			SHADER_PARAMETER_RDG_TEXTURE_UAV(RWTexture2D, RWSelectionMap)
+			END_SHADER_PARAMETER_STRUCT()
+
+			static void ModifyCompilationEnvironment(const FGlobalShaderPermutationParameters& InParameters, FShaderCompilerEnvironment& OutEnvironment)
+		{
+			FGlobalShader::ModifyCompilationEnvironment(InParameters, OutEnvironment);
+			OutEnvironment.SetDefine(TEXT("THREAD_GROUP_SIZE"), NON_LOCAL_MEAN_THREAD_GROUP_SIZE);
+		}
+	};
+
+	class FCombineFilteredImageCS : public FGlobalShader
+	{
+		DECLARE_GLOBAL_SHADER(FCombineFilteredImageCS)
+			SHADER_USE_PARAMETER_STRUCT(FCombineFilteredImageCS, FGlobalShader)
+
+	public:
+		BEGIN_SHADER_PARAMETER_STRUCT(FParameters, )
+			SHADER_PARAMETER_RDG_TEXTURE_SRV_ARRAY(Texture2D, FilteredImages, [2])
+			SHADER_PARAMETER_RDG_TEXTURE_SRV(Texture2D, SelectionMap)
+			SHADER_PARAMETER(FIntPoint, TextureSize)
+			SHADER_PARAMETER_RDG_TEXTURE_UAV(RWTexture2D, RWFilteredImage)
+			END_SHADER_PARAMETER_STRUCT()
+
+			static void ModifyCompilationEnvironment(const FGlobalShaderPermutationParameters& InParameters, FShaderCompilerEnvironment& OutEnvironment)
+		{
+			FGlobalShader::ModifyCompilationEnvironment(InParameters, OutEnvironment);
+			OutEnvironment.SetDefine(TEXT("THREAD_GROUP_SIZE"), NON_LOCAL_MEAN_THREAD_GROUP_SIZE);
+		}
+	};
+
+	//--------------------------------------------------------------------------------------------------------------------
+	// NFOR texture description and main denoising entry
+
+	struct FNFORTextureDesc
+	{
+		FRDGTextureRef Image;
+
+		/** the specific channels for this feature*/
+		int32		  ChannelOffset;
+		int32		  ChannelCount;
+
+		/** The number of channels of the image*/
+		int32		  NumOfChannel;
+
+		FNFORTextureDesc(FRDGTextureRef Image, int32 ChannelOffset = 0, int32 ChannelCount = 4, int32 NumOfChannel = 4) :
+			Image(Image), ChannelOffset(ChannelOffset), ChannelCount(ChannelCount), NumOfChannel(NumOfChannel) {}
+
+	};
+
+	struct FFeatureDesc
+	{
+		union
+		{
+			FNFORTextureDesc Feature;
+			FNFORTextureDesc Data;
+		};
+
+		FNFORTextureDesc Variance;
+
+		/** Variance type, greyscale, normal, or colored */
+		EVarianceType VarianceType;
+
+		/**Indicate no need to denoise if true**/
+		bool bCleanFeature;
+
+		FFeatureDesc(FNFORTextureDesc Feature, FNFORTextureDesc Variance, EVarianceType VarianceType = EVarianceType::GreyScale, bool bCleanFeature = false)
+			:Feature(Feature), Variance(Variance), VarianceType(VarianceType), bCleanFeature(bCleanFeature) {}
+	};
+
+	typedef FFeatureDesc FRadianceDesc;
+
+	// Denoise features based on non-local mean on its variance texture.
+	void FilterFeatures(
+		FRDGBuilder& GraphBuilder,
+		const FSceneView& View,
+		const TArray<FFeatureDesc>& FeatureDescs);
+
+	// Denoise the radiance based on spatial temporal features.
+	bool FilterMain(
+		FRDGBuilder& GraphBuilder,
+		const FSceneView& View,
+		const TArray<FRadianceDesc>& Radiances,
+		const TArray<FFeatureDesc>& FeatureDescs,
+		const FRDGTextureRef& DenoisedRadiance);
+
+}
\ No newline at end of file
diff --git a/Engine/Plugins/Experimental/NFORDenoise/Source/NFORDenoise/Private/NFORRegressionCPUSolver.cpp b/Engine/Plugins/Experimental/NFORDenoise/Source/NFORDenoise/Private/NFORRegressionCPUSolver.cpp
new file mode 100644
--- /dev/null
+++ b/Engine/Plugins/Experimental/NFORDenoise/Source/NFORDenoise/Private/NFORRegressionCPUSolver.cpp
@@ -0,0 +1,581 @@
+// Copyright Epic Games, Inc. All Rights Reserved.
+
+#include "NFORRegressionCPUSolver.h"
+#include "NFORWeightedLSRCommon.h"
+
+#include "RenderGraphUtils.h"
+#include "RHIGPUReadback.h"
+
+#include "HAL/IConsoleManager.h"
+#include "Math/Quat.h"
+
+// Just to be sure, also added this in Eigen.Build.cs
+#ifndef EIGEN_MPL2_ONLY
+#define EIGEN_MPL2_ONLY
+#endif
+
+#if defined(_MSC_VER) && USING_CODE_ANALYSIS
+#pragma warning(push)
+#pragma warning(disable:6294) // Ill-defined for-loop:  initial condition does not satisfy test.  Loop body not executed.
+#endif
+PRAGMA_DEFAULT_VISIBILITY_START
+THIRD_PARTY_INCLUDES_START
+#include <Eigen/Dense>
+THIRD_PARTY_INCLUDES_END
+PRAGMA_DEFAULT_VISIBILITY_END
+#if defined(_MSC_VER) && USING_CODE_ANALYSIS
+#pragma warning(pop)
+#endif
+
+namespace NFORRegressionSolverCPU
+{
+	TAutoConsoleVariable<int32> CVarNFORRegressionRecombineType(
+		TEXT("r.NFOR.CPU.Regression.Recombine.Type"),
+		0,
+		TEXT("0: Apply patch reconstruction to patch for current frame, weighted sum of neighbor patch to the current pixel for other frames.")
+		TEXT("1: Apply weighted sum of neighbor patch to the current pixel for all frames."),
+		ECVF_RenderThreadSafe);
+
+	enum class ERecombineType : int32
+	{
+		Auto,
+		Gather,
+		MAX
+	};
+
+	ERecombineType GetRecombineType()
+	{
+		const int32 RecombineType = FMath::Clamp(
+			CVarNFORRegressionRecombineType.GetValueOnAnyThread(),
+			static_cast<int32>(ERecombineType::Auto),
+			static_cast<int32>(ERecombineType::MAX)-1);
+
+		return static_cast<ERecombineType>(RecombineType);
+	}
+
+	struct NFORFirstOrderRegressionState
+	{
+		TArray<FLinearColor> Y; // WxHxC
+		TArray<float>		 X; // WxHxF, F is the number of feature channel
+		TArray<float>		 W; // WxHxN, N is the total number of weights over temporal frames = (2*PatchDistance+1)^2 * T
+
+		TArray<FLinearColor> Buffer; // WxHxC
+
+		FIntPoint CurrentSize = FIntPoint::ZeroValue;
+		int CurrentNumChannels = 0;
+		int CurrentNumWeights = 0;
+		int CurrentNumRadianceChannels = 0;
+		FIntPoint CurrentOffset = FIntPoint::ZeroValue;
+		int NumOfFrames = 1;
+
+		// Size: Texture image size Y
+		void Update(FIntPoint TextureSize, int NumFeatureChannel, int NumWeights, int NumRadianceChannels);
+		void Update(const FWeightedLSRDesc& WeightedLSRDesc);
+	};
+
+	void Apply(NFORFirstOrderRegressionState& FirstOrderRegressionState);
+
+	struct FNFORLinearSolve
+	{
+		TArray<float> A; // N*FxF
+		TArray<float> B; // N*FxC
+
+		int CurrentNumElements = 0;
+		int CurrentNumFeatureChannel = 0;
+		int CurrentNumRadianceChannels = 0;
+
+		void Update(int NumOfElements, int NumFeatureChannel, int NumRadianceChannels);
+	};
+
+	void Solve(FNFORLinearSolve& NFORLinearSolve);
+
+	void NFORFirstOrderRegressionState::Update(FIntPoint Size, int NumFeatureChannel, int NumWeights, int NumRadianceChannels)
+	{
+		int NewSize = Size.X * Size.Y;
+		const bool bImageSizeChanged = Y.Num() != NewSize;
+		const bool bNumFeatureChannelChanged = CurrentNumChannels != NumFeatureChannel;
+		const bool bNumWeightsChanged = CurrentNumWeights != NumWeights;
+
+		if (bImageSizeChanged)
+		{
+			Y.SetNumUninitialized(NewSize);
+			Buffer.SetNumUninitialized(NewSize);
+		}
+
+		if (bImageSizeChanged || bNumFeatureChannelChanged)
+		{
+			X.SetNumUninitialized(NewSize * NumFeatureChannel);
+		}
+
+		if (bImageSizeChanged || bNumWeightsChanged)
+		{
+			W.SetNumUninitialized(NewSize * NumWeights);
+		}
+
+		CurrentSize = Size;
+		CurrentNumChannels = NumFeatureChannel;
+		CurrentNumWeights = NumWeights;
+		CurrentNumRadianceChannels = NumRadianceChannels;
+	}
+
+	template <typename MatrixType>
+	void PrintEigenMatrix(const MatrixType& matrix) {
+		for (int i = 0; i < matrix.rows(); ++i) {
+			for (int j = 0; j < matrix.cols(); ++j) {
+				UE_LOG(LogTemp, Log, TEXT("%.5f "), static_cast<float>(matrix(i, j)));
+			}
+			UE_LOG(LogTemp, Log, TEXT("\n"));
+		}
+	}
+
+	void NFORFirstOrderRegressionState::Update(const FWeightedLSRDesc& WeightedLSRDesc)
+	{
+
+		int NewSize = WeightedLSRDesc.TextureSize.X * WeightedLSRDesc.TextureSize.Y;
+		const bool bImageSizeChanged = Y.Num() != (NewSize * WeightedLSRDesc.NumOfFrames);
+		const bool bNumFeatureChannelChanged = CurrentNumChannels != WeightedLSRDesc.NumOfFeatureChannels;
+		const bool bNumWeightsChanged = CurrentNumWeights != WeightedLSRDesc.NumOfWeightsPerPixel;
+		const bool bOffsetChanged = CurrentOffset != WeightedLSRDesc.Offset;
+
+		if (bImageSizeChanged)
+		{
+			Y.SetNumUninitialized(NewSize * WeightedLSRDesc.NumOfFrames);
+			Buffer.SetNumUninitialized(NewSize);
+		}
+
+		if (bImageSizeChanged || bNumFeatureChannelChanged)
+		{
+			X.SetNumUninitialized(NewSize * WeightedLSRDesc.NumOfFeatureChannels);
+		}
+
+
+		int WeightSize = NewSize;
+		if (WeightedLSRDesc.SolverType == EWeightedLSRSolverType::Tiled)
+		{
+			WeightSize = WeightedLSRDesc.Width * WeightedLSRDesc.Height;
+
+		}
+
+		if (bImageSizeChanged || bNumWeightsChanged)
+		{
+			W.SetNumUninitialized(WeightSize * WeightedLSRDesc.NumOfWeightsPerPixel);
+		}
+
+		CurrentSize = WeightedLSRDesc.TextureSize;
+		CurrentNumChannels = WeightedLSRDesc.NumOfFeatureChannels;
+		CurrentNumWeights = WeightedLSRDesc.NumOfWeightsPerPixel;
+		CurrentNumRadianceChannels = WeightedLSRDesc.NumOfRadianceChannels;
+		CurrentOffset = WeightedLSRDesc.Offset;
+		NumOfFrames = WeightedLSRDesc.NumOfFrames;
+	}
+
+	FIntPoint ClampPointMirrored(FIntPoint P, FIntPoint Size)
+	{
+		FIntPoint SizeMax = Size - 1;
+
+		int32 X = FMath::Abs(SizeMax.X - FMath::Abs(SizeMax.X - P.X));
+		int32 Y = FMath::Abs(SizeMax.Y - FMath::Abs(SizeMax.Y - P.Y));
+
+		return FIntPoint(X, Y);
+	}
+
+	void Apply(NFORFirstOrderRegressionState& FirstOrderRegressionState)
+	{
+		// Set up X, Y, W and solve with Eigen library
+		// Input:
+		// W = NonLocalMean weights
+		// X = feature dimension
+		// Y = Radiance
+
+		// This function solves a weighted linear regression problem
+		// Predicting a weight beta such that XB predicts the measured image value Y.
+		// for each point
+		// Solve Loss	= \sum wi (yi-xi^t\betai)^2
+		//				= (Y-XB) W (Y-XB)
+		// where:
+		//	X = Xp or Xp-Xc
+		// 
+		// Solution:
+		//  B = (X^T W X)^{-1} X^T W Y
+		// Reconstruct = X B
+
+		/*
+		example code to solve the weighted linear regression for a patch
+
+		int32 N = FirstOrderRegressionState.CurrentNumWeights;
+		int32 D = FirstOrderRegressionState.CurrentNumChannels;
+		Eigen::VectorXf W(N);
+		Eigen::MatrixXf Y(N, 3);
+		Eigen::MatrixXf X(N, D);
+
+		Eigen::VectorXf WSqrt = W.cwiseSqrt();
+		Eigen::VectorXf WDiag = WSqrt.asDiagonal();
+		Eigen::MatrixXf Beta = (X.transpose() * WDiag * X).colPivHouseholderQr().solve(X.transpose() * WDiag * Y);
+		Eigen::MatrixXf Reconstruct = X * Beta;
+
+		call the code above for each pixel and write to the target buffer for value and weights accumulation. Each pixel is a weighted
+		sum of NumWeights values.
+		*/
+		NFORFirstOrderRegressionState& Ctx = FirstOrderRegressionState;
+		Ctx.Buffer.Init(FLinearColor::Transparent, Ctx.Y.Num());
+
+		int32 N = Ctx.CurrentNumWeights;
+		int32 D = Ctx.CurrentNumChannels;
+		int32 B = Ctx.CurrentNumRadianceChannels;
+		int32 T = Ctx.NumOfFrames;
+		int32 R = (FMath::Sqrt(float(N / T)) - 1) / 2;
+		int32 n = N / T;
+
+		Eigen::VectorXf W = Eigen::VectorXf::Zero(N);
+		Eigen::MatrixXf Y = Eigen::MatrixXf::Zero(N, 3);
+		Eigen::MatrixXf X = Eigen::MatrixXf::Zero(N, D / T);
+
+		int XPadding = Ctx.CurrentOffset.X;
+		int YPadding = Ctx.CurrentOffset.Y;
+		int32 SingleFrameSize = (Ctx.CurrentSize.X - 2 * XPadding) * (Ctx.CurrentSize.Y - 2 * XPadding);
+
+		ERecombineType RecombineType = GetRecombineType();
+
+		bool bPrintDebug = true;
+
+		//TODO: Parallelize on GPU thread.
+		for (int Py = YPadding; Py < Ctx.CurrentSize.Y - YPadding; ++Py)
+		{
+			for (int Px = XPadding; Px < Ctx.CurrentSize.X - XPadding; ++Px)
+			{
+
+				// For each pixel, solve the WLR.
+				FIntPoint P = FIntPoint(Px, Py);
+				int32 PIndex = (P.X + P.Y * Ctx.CurrentSize.X);
+
+				{
+					// The weight index is stored in the internal rectangular.
+					FIntPoint WP = P - Ctx.CurrentOffset;
+
+					int32 WPIndex = (WP.X + WP.Y * (Ctx.CurrentSize.X - 2 * XPadding));
+
+					for (int i = 0; i < N; ++i)
+					{
+						int32 WOffset = i % n;
+						int32 TOffset = i / n;
+						int32 SerializedWPIndex = SingleFrameSize * n * TOffset + WPIndex * n + WOffset;
+
+						W(i) = Ctx.W[SerializedWPIndex];
+					}
+				}
+
+				// Fill the data.
+				for (int Qy = Py - R; Qy <= Py + R; ++Qy)
+				{
+					for (int Qx = Px - R; Qx <= Px + R; ++Qx)
+					{
+						FIntPoint Q = FIntPoint(Qx, Qy);
+						FIntPoint QMirrored = ClampPointMirrored(Q, Ctx.CurrentSize);
+						int32 QIndex = (QMirrored.X + QMirrored.Y * Ctx.CurrentSize.X);
+
+
+						FIntPoint LocalQ = Q - (P - FIntPoint(R, R));
+						int LocalIndex = (LocalQ.X + LocalQ.Y * (2 * R + 1));
+
+						for (int ti = 0; ti < T; ++ti)
+						{
+							float* YPtr = ((float*)Ctx.Y.GetData()) + (QIndex * T + ti) * B;
+							Y(LocalIndex + ti * n, 0) = *YPtr;
+							Y(LocalIndex + ti * n, 1) = *(YPtr + 1);
+							Y(LocalIndex + ti * n, 2) = *(YPtr + 2);
+
+						}
+
+
+						for (int di = 0; di < D; ++di)
+						{
+							int FOffset = di % (D / T);
+							int TOffset = di / (D / T);
+							X(LocalIndex + TOffset * n, FOffset) = Ctx.X[QIndex * D + di];
+						}
+					}
+				}
+
+				// Solve for the whole patch.
+				Eigen::VectorXf WSqrt = W.cwiseSqrt();
+				Eigen::MatrixXf WDiag = WSqrt.asDiagonal();
+				Eigen::MatrixXf AMatrix = X.transpose() * WDiag * X;
+				Eigen::MatrixXf BMatrix = X.transpose() * WDiag * Y;
+				Eigen::MatrixXf Beta = (AMatrix).colPivHouseholderQr().solve(BMatrix);
+				Eigen::MatrixXf Reconstruct = X * Beta;
+
+				// Write back: Two strategy:
+				// One:
+				//		1. Write back the current patch image to each pixel in the patch for the current frame
+				//			Find the best fit for the patch given the similarity to the current center patch. 
+				//			Each pixel is a weighted sum of denoised self based on all surrounding blocks.
+				//		2. for history, gather all pixel contribution to the current center pixel.
+				// Two:
+				//		For all frames, gather all pixel contribution to the current pixel.
+
+				for (int Qy = Py - R; Qy <= Py + R; ++Qy)
+				{
+					for (int Qx = Px - R; Qx <= Px + R; ++Qx)
+					{
+						FIntPoint Q = FIntPoint(Qx, Qy);
+						FIntPoint QMirrored = ClampPointMirrored(Q, Ctx.CurrentSize);
+						int32 QIndex = (QMirrored.X + QMirrored.Y * Ctx.CurrentSize.X);
+
+						FIntPoint LocalQ = Q - (P - FIntPoint(R, R));
+						int LocalIndex = (LocalQ.X + LocalQ.Y * (2 * R + 1));
+
+						for (int i = 0; i < T; ++i)
+						{
+							int TemporalLocalIndex = LocalIndex + i * n;
+
+							if (RecombineType == ERecombineType::Auto)
+							{
+								if (i == T / 2)
+								{
+									Ctx.Buffer[QIndex] += W(TemporalLocalIndex) *
+										FLinearColor(Reconstruct(TemporalLocalIndex, 0), Reconstruct(TemporalLocalIndex, 1), Reconstruct(TemporalLocalIndex, 2), 1);
+								}
+								else
+								{
+									Ctx.Buffer[PIndex] += W(TemporalLocalIndex) *
+										FLinearColor(Reconstruct(TemporalLocalIndex, 0), Reconstruct(TemporalLocalIndex, 1), Reconstruct(TemporalLocalIndex, 2), 1);
+								}
+							}
+							else
+							{
+								Ctx.Buffer[PIndex] += W(TemporalLocalIndex) *
+									FLinearColor(Reconstruct(TemporalLocalIndex, 0), Reconstruct(TemporalLocalIndex, 1), Reconstruct(TemporalLocalIndex, 2), 1);
+							}
+						}
+					}
+				} // Write back
+			}
+		}
+
+		// Copy back from the buffer.
+		for (int Py = 0; Py < Ctx.CurrentSize.Y; ++Py)
+		{
+			for (int Px = 0; Px < Ctx.CurrentSize.X; ++Px)
+			{
+				FIntPoint P = FIntPoint(Px, Py);
+				int32 PIndex = (P.X + P.Y * Ctx.CurrentSize.X);
+
+				Ctx.Y[PIndex] = Ctx.Buffer[PIndex];
+			}
+		}
+
+	}
+
+	void FNFORLinearSolve::Update(int NumOfElements, int NumFeatureChannel, int NumRadianceChannels)
+	{
+		if (CurrentNumElements != NumOfElements ||
+			CurrentNumFeatureChannel != NumFeatureChannel)
+		{
+			CurrentNumElements = NumOfElements;
+			CurrentNumFeatureChannel = NumFeatureChannel;
+
+			A.SetNumUninitialized(CurrentNumElements * CurrentNumFeatureChannel * CurrentNumFeatureChannel);
+		}
+
+		if (CurrentNumElements != NumOfElements ||
+			CurrentNumFeatureChannel != NumFeatureChannel ||
+			CurrentNumRadianceChannels != NumRadianceChannels)
+		{
+			CurrentNumElements = NumOfElements;
+			CurrentNumFeatureChannel = NumFeatureChannel;
+			CurrentNumRadianceChannels = NumRadianceChannels;
+
+			B.SetNumUninitialized(CurrentNumElements * CurrentNumFeatureChannel * CurrentNumRadianceChannels);
+		}
+	}
+
+	void Solve(FNFORLinearSolve& NFORLinearSolve)
+	{
+
+		const int F = NFORLinearSolve.CurrentNumFeatureChannel;
+		const int C = NFORLinearSolve.CurrentNumRadianceChannels;
+		const int N = NFORLinearSolve.CurrentNumElements;
+
+		const int ASize = F * F;
+		const int BSize = F * C;
+
+		Eigen::MatrixXf A = Eigen::MatrixXf::Zero(F, F);
+		Eigen::MatrixXf B = Eigen::MatrixXf::Zero(F, C);
+
+		for (int i = 0; i < N; ++i)
+		{
+			for (int j = 0; j < ASize; ++j)
+			{
+				A(j / F, j % F) = NFORLinearSolve.A[i * ASize + j];
+			}
+
+			for (int j = 0; j < BSize; ++j)
+			{
+				B(j / C, j % C) = NFORLinearSolve.B[i * BSize + j];
+			}
+
+			Eigen::MatrixXf Beta = (A).colPivHouseholderQr().solve(B);
+
+			for (int j = 0; j < BSize; ++j)
+			{
+				NFORLinearSolve.B[i * BSize + j] = Beta(j / C, j % C);
+			}
+		}
+	}
+}
+
+template <typename PixelType>
+static void CopyTextureFromGPUToCPU(FRHICommandListImmediate& RHICmdList, FRHITexture* SrcTexture, FIntPoint Size, TArray<PixelType>& DstArray)
+{
+	FRHIGPUTextureReadback Readback(TEXT("NFOR::TextureReadback"));
+	Readback.EnqueueCopy(RHICmdList, SrcTexture, FIntVector::ZeroValue, 0, FIntVector(Size.X, Size.Y, 1));
+	RHICmdList.BlockUntilGPUIdle();
+
+	int32_t SrcStride = 0;
+	const PixelType* SrcBuffer = static_cast<PixelType*>(Readback.Lock(SrcStride, nullptr));
+
+	PixelType* DstBuffer = DstArray.GetData();
+	for (int Y = 0; Y < Size.Y; Y++, DstBuffer += Size.X, SrcBuffer += SrcStride)
+	{
+		FPlatformMemory::Memcpy(DstBuffer, SrcBuffer, Size.X * sizeof(PixelType));
+
+	}
+	Readback.Unlock();
+}
+
+template <typename PixelType>
+static void CopyTextureFromCPUToGPU(FRHICommandListImmediate& RHICmdList, const TArray<PixelType>& SrcArray, FIntPoint Size, FRHITexture* DstTexture)
+{
+	uint32_t DestStride;
+	FLinearColor* DstBuffer = static_cast<PixelType*>(RHICmdList.LockTexture2D(DstTexture, 0, RLM_WriteOnly, DestStride, false));
+	DestStride /= sizeof(PixelType);
+	const FLinearColor* SrcBuffer = SrcArray.GetData();
+	for (int Y = 0; Y < Size.Y; Y++, SrcBuffer += Size.X, DstBuffer += DestStride)
+	{
+		FPlatformMemory::Memcpy(DstBuffer, SrcBuffer, Size.X * sizeof(PixelType));
+	}
+	RHICmdList.UnlockTexture2D(DstTexture, 0, false);
+}
+
+template <typename ElementType>
+static void CopyBufferFromGPUToCPU(FRHICommandListImmediate& RHICmdList, FRHIBuffer* SrcBuffer, uint32 NumElement, TArray<ElementType>& DstArray)
+{
+	FRHIGPUBufferReadback Readback(TEXT("NFOR::BufferReadBack"));
+	uint32_t NumBytes = NumElement * sizeof(ElementType);
+	Readback.EnqueueCopy(RHICmdList, SrcBuffer, NumBytes);
+	RHICmdList.BlockUntilGPUIdle();
+
+	const ElementType* SrcRawBuffer = static_cast<ElementType*>(Readback.Lock(NumBytes));
+
+	ElementType* DstBuffer = DstArray.GetData();
+	FPlatformMemory::Memcpy(DstBuffer, SrcRawBuffer, NumBytes);
+	Readback.Unlock();
+}
+
+template <typename ElementType>
+static void CopyBufferFromCPUToGPU(FRHICommandListImmediate& RHICmdList, TArray<ElementType>& SrcArray, uint32 NumElement, FRHIBuffer* DstBuffer)
+{
+	uint32_t NumBytes = NumElement * sizeof(ElementType);
+	void* Data = RHICmdList.LockBuffer(DstBuffer, 0, NumBytes, RLM_WriteOnly);
+	ElementType* DstRawBuffer = static_cast<ElementType*>(Data);
+
+	const ElementType* SrcBuffer = SrcArray.GetData();
+	FPlatformMemory::Memcpy(DstRawBuffer, SrcBuffer, NumBytes);
+	RHICmdList.UnlockBuffer(DstBuffer);
+}
+
+// CPU solver relevants used for verification.	
+BEGIN_SHADER_PARAMETER_STRUCT(FFirstOrderRegressionParameters, )
+	RDG_BUFFER_ACCESS(X, ERHIAccess::CopySrc)
+	RDG_TEXTURE_ACCESS(Y, ERHIAccess::CopySrc)
+	RDG_BUFFER_ACCESS(Ys, ERHIAccess::CopySrc)
+	RDG_BUFFER_ACCESS(W, ERHIAccess::CopySrc)
+	RDG_TEXTURE_ACCESS(Reconstruct, ERHIAccess::CopyDest)
+	SHADER_PARAMETER(int32, T)
+END_SHADER_PARAMETER_STRUCT()
+
+BEGIN_SHADER_PARAMETER_STRUCT(FLinearSolverParameters, )
+	RDG_BUFFER_ACCESS(A, ERHIAccess::CopySrc)
+	RDG_BUFFER_ACCESS(B, ERHIAccess::CopySrc)
+	RDG_BUFFER_ACCESS(X, ERHIAccess::CopyDest)
+END_SHADER_PARAMETER_STRUCT()
+
+static NFORRegressionSolverCPU::NFORFirstOrderRegressionState RegressionState;
+static NFORRegressionSolverCPU::FNFORLinearSolve NFORLinearSolveState;
+
+void SolveWeightedLSRCPU(
+	FRDGBuilder& GraphBuilder,
+	const FSceneView& View,
+	const FRDGBufferRef& Feature,
+	const FRDGTextureRef& Radiance,
+	const FRDGBufferRef& NonLocalMeanWeightsBuffer,
+	const FRDGTextureRef& FilteredRadiance,
+	const FWeightedLSRDesc& WeightedLSRDesc,
+	const FRDGBufferRef Radiances,
+	const FRDGTextureRef& SourceAlbedo)
+{
+	FIntPoint TextureSize = Radiance->Desc.Extent;
+	FFirstOrderRegressionParameters* FirstOrderRegressionParameters = GraphBuilder.AllocParameters<FFirstOrderRegressionParameters>();
+
+	{
+		FirstOrderRegressionParameters->X = Feature;
+		FirstOrderRegressionParameters->Y = Radiance;
+		FirstOrderRegressionParameters->Ys = Radiances;
+		FirstOrderRegressionParameters->W = NonLocalMeanWeightsBuffer;
+		FirstOrderRegressionParameters->Reconstruct = FilteredRadiance;
+	}
+
+	// Need to read GPU mask outside Pass function, as the value is not refreshed inside the pass
+	GraphBuilder.AddPass(RDG_EVENT_NAME("SolvWeightedLSR (Eigen CPU Sequential)"), FirstOrderRegressionParameters, ERDGPassFlags::Readback,
+		[FirstOrderRegressionParameters, WeightedLSRDesc, TextureSize](FRHICommandListImmediate& RHICmdList)
+		{
+			RegressionState.Update(WeightedLSRDesc);
+
+			// Read W,Y and X to CPU	
+			CopyBufferFromGPUToCPU(RHICmdList, FirstOrderRegressionParameters->Ys->GetRHI(),
+				FirstOrderRegressionParameters->Ys->GetSize() / FirstOrderRegressionParameters->Ys->GetStride(), RegressionState.Y);
+			CopyBufferFromGPUToCPU(RHICmdList, FirstOrderRegressionParameters->X->GetRHI(),
+				FirstOrderRegressionParameters->X->GetSize() / FirstOrderRegressionParameters->X->GetStride(), RegressionState.X);
+			CopyBufferFromGPUToCPU(RHICmdList, FirstOrderRegressionParameters->W->GetRHI(),
+				FirstOrderRegressionParameters->W->GetSize() / FirstOrderRegressionParameters->W->GetStride(), RegressionState.W);
+
+			// Run the linear regression
+			NFORRegressionSolverCPU::Apply(RegressionState);
+
+			// Copy back
+			CopyTextureFromCPUToGPU(RHICmdList, RegressionState.Y, TextureSize, FirstOrderRegressionParameters->Reconstruct->GetRHI()->GetTexture2D());
+		});
+}
+
+void SolveLinearEquationCPU(
+	FRDGBuilder& GraphBuilder, 
+	const FRDGBufferRef& A, 
+	const FRDGBufferRef& B,
+	const int32 NumOfElements,
+	const FIntPoint BDim,
+	const FRDGBufferRef& X)
+{
+	FLinearSolverParameters* LinearSolverParameters = GraphBuilder.AllocParameters<FLinearSolverParameters>();
+	LinearSolverParameters->A = A;
+	LinearSolverParameters->B = B;
+	LinearSolverParameters->X = X;
+
+	GraphBuilder.AddPass(RDG_EVENT_NAME("ReconstructWeights (Eigen CPU Sequential)"), LinearSolverParameters, ERDGPassFlags::Readback,
+		[LinearSolverParameters, NumOfElements, BDim](FRHICommandListImmediate& RHICmdList)
+		{
+			NFORLinearSolveState.Update(NumOfElements, BDim.X, BDim.Y);
+
+			// Copy W,Y and X to CPU			
+			CopyBufferFromGPUToCPU(RHICmdList, LinearSolverParameters->A->GetRHI(),
+				LinearSolverParameters->A->GetSize() / LinearSolverParameters->A->GetStride(), NFORLinearSolveState.A);
+
+			CopyBufferFromGPUToCPU(RHICmdList, LinearSolverParameters->B->GetRHI(),
+				LinearSolverParameters->B->GetSize() / LinearSolverParameters->B->GetStride(), NFORLinearSolveState.B);
+
+			NFORRegressionSolverCPU::Solve(NFORLinearSolveState);
+
+			// Copy back
+			CopyBufferFromCPUToGPU(RHICmdList, NFORLinearSolveState.B,
+				LinearSolverParameters->B->GetSize() / LinearSolverParameters->B->GetStride(), LinearSolverParameters->X->GetRHI());
+		});
+}
\ No newline at end of file
diff --git a/Engine/Plugins/Experimental/NFORDenoise/Source/NFORDenoise/Private/NFORRegressionCPUSolver.h b/Engine/Plugins/Experimental/NFORDenoise/Source/NFORDenoise/Private/NFORRegressionCPUSolver.h
new file mode 100644
--- /dev/null
+++ b/Engine/Plugins/Experimental/NFORDenoise/Source/NFORDenoise/Private/NFORRegressionCPUSolver.h
@@ -0,0 +1,29 @@
+// Copyright Epic Games, Inc. All Rights Reserved.
+
+#pragma once
+
+#include "RenderGraphFwd.h"
+#include "Math/MathFwd.h"
+
+struct FWeightedLSRDesc;
+class FSceneView;
+
+void SolveWeightedLSRCPU(
+	FRDGBuilder& GraphBuilder,
+	const FSceneView& View,
+	const FRDGBufferRef& Feature,
+	const FRDGTextureRef& Radiance,
+	const FRDGBufferRef& NonLocalMeanWeightsBuffer,
+	const FRDGTextureRef& FilteredRadiance,
+	const FWeightedLSRDesc& WeightedLSRDesc,
+	const FRDGBufferRef Radiances,
+	const FRDGTextureRef& SourceAlbedo);
+
+// Solve AX=B
+void SolveLinearEquationCPU(
+	FRDGBuilder& GraphBuilder,
+	const FRDGBufferRef& A,
+	const FRDGBufferRef& B,
+	const int32 NumOfElements,
+	const FIntPoint BDim,
+	const FRDGBufferRef& X);
\ No newline at end of file
diff --git a/Engine/Plugins/Experimental/NFORDenoise/Source/NFORDenoise/Private/NFORWeightedLSRCommon.h b/Engine/Plugins/Experimental/NFORDenoise/Source/NFORDenoise/Private/NFORWeightedLSRCommon.h
new file mode 100644
--- /dev/null
+++ b/Engine/Plugins/Experimental/NFORDenoise/Source/NFORDenoise/Private/NFORWeightedLSRCommon.h
@@ -0,0 +1,31 @@
+// Copyright Epic Games, Inc. All Rights Reserved.
+
+#pragma once
+
+#include "Math/IntPoint.h"
+#include "Math/MathFwd.h"
+
+enum class EWeightedLSRSolverType : uint8
+{
+	Full,			// whole image input
+	Tiled			// tiled image input, Overlapping region data are included
+};
+
+struct FWeightedLSRDesc
+{
+	int32 NumOfFeatureChannels; // F
+	int32 NumOfFeatureChannelsPerFrame;
+	int32 NumOfWeightsPerPixel; // N
+	int32 NumOfWeightsPerPixelPerFrame;
+	int32 NumOfRadianceChannels;
+	int32 NumOfRadianceChannelsPerFrame;
+	int32 Width;				// W
+	int32 Height;				// H
+
+	FIntPoint TileStartPosition;	// The start position of the tile 
+	FIntPoint Offset;				// Used in tiled mode as we need to copy data outside of the region
+
+	int32 NumOfFrames;			// T
+	FIntPoint TextureSize;		// (W+2B,H+2B) where B is the border size
+	EWeightedLSRSolverType SolverType;
+};
diff --git a/Engine/Plugins/Experimental/NFORDenoise/Source/NFORDenoise/Public/NFORDenoise.h b/Engine/Plugins/Experimental/NFORDenoise/Source/NFORDenoise/Public/NFORDenoise.h
new file mode 100644
--- /dev/null
+++ b/Engine/Plugins/Experimental/NFORDenoise/Source/NFORDenoise/Public/NFORDenoise.h
@@ -0,0 +1,18 @@
+// Copyright Epic Games, Inc. All Rights Reserved.
+
+#pragma once
+
+#include "CoreMinimal.h"
+#include "Logging/LogMacros.h"
+#include "Modules/ModuleInterface.h"
+#include "Modules/ModuleManager.h"
+
+DECLARE_LOG_CATEGORY_EXTERN(LogNFORDenoise, Log, All);
+
+class FNFORDenoiseModule : public IModuleInterface
+{
+public:
+	/** IModuleInterface implementation */
+	virtual void StartupModule() override;
+	virtual void ShutdownModule() override;
+};
diff --git a/Engine/Plugins/Experimental/NFORDenoise/Source/NFORDenoise/Shaders/NFORDenoise.usf b/Engine/Plugins/Experimental/NFORDenoise/Source/NFORDenoise/Shaders/NFORDenoise.usf
new file mode 100644
--- /dev/null
+++ b/Engine/Plugins/Experimental/NFORDenoise/Source/NFORDenoise/Shaders/NFORDenoise.usf
@@ -0,0 +1,1065 @@
+// Copyright Epic Games, Inc. All Rights Reserved.
+
+#include "/Engine/Public/Platform.ush"
+#include "/Engine/Private/Common.ush"
+#include "/Engine/Private/ScreenPass.ush"
+
+#ifndef SOURCE_CHANNEL_COUNT
+#define SOURCE_CHANNEL_COUNT 4
+#endif
+
+#if SOURCE_CHANNEL_COUNT == 1
+#define TPixelValue float
+#elif SOURCE_CHANNEL_COUNT == 2
+#define TPixelValue float2
+#elif SOURCE_CHANNEL_COUNT == 3
+#define TPixelValue float3
+#elif SOURCE_CHANNEL_COUNT == 4
+#define TPixelValue float4
+#endif
+
+#define IMAGE_VARINCE_NORMAL		0
+#define IMAGE_VARIANCE_GREYSCALE	1
+#define IMAGE_VARIANCE_COLORED		2
+
+#ifndef IMAGE_VARIANCE_TYPE
+#define IMAGE_VARIANCE_TYPE IMAGE_VARIANCE_GREYSCALE
+#endif
+
+#define PRE_ALBEDO_DIVIDE_DISABLED	0
+#define PRE_ALBEDO_DIVIDE_EACH		1
+#define PRE_ALBEDO_DIVIDE_FINAL		2
+
+#ifndef PRE_ALBEDO_DIVIDE
+#define PRE_ALBEDO_DIVIDE PRE_ALBEDO_DIVIDE_DISABLED
+#endif
+
+#ifndef APPEND_CONSTANT_DIMENSION_TO_X
+#define APPEND_CONSTANT_DIMENSION_TO_X 0
+#endif
+
+#ifndef NONLOCALMEAN_SEPARATE_SOURCE
+#define NONLOCALMEAN_SEPARATE_SOURCE 0
+#endif
+
+#define NONLOCALMEAN_PATCHONLY		0
+#define NONLOCALMEAN_DISTACNE_PATCH 1
+#define NONLOCALMEAN_WEIGHT_TYPE	NONLOCALMEAN_PATCHONLY
+
+#ifndef NUM_FEATURE
+#define NUM_FEATURE 0
+#endif
+
+#ifndef SMALL_MATRIX_OPTIMIZE
+#define SMALL_MATRIX_OPTIMIZE 0
+#endif
+
+// Sampling steps for performance optimization.
+#ifndef USE_SAMPLING_STEP
+#define USE_SAMPLING_STEP 0
+#endif
+
+#define SAMPLING_TYPE_INCREMENTAL			0
+#define SAMPLING_TYPE_FOCUS_CURRENT_FRAME	1
+
+#define SAMPLING_TYPE SAMPLING_TYPE_FOCUS_CURRENT_FRAME
+
+// Matrix multiplication
+#define WEIGHTED_MULTIPLICATION_QUADRATIC	0	//X^TWX
+#define WEIGHTED_MULTIPLICATION_GENERALIZED	1	//X^TWY
+
+#ifndef WEIGHTED_MULTIPLICATION_TYPE
+#define WEIGHTED_MULTIPLICATION_TYPE WEIGHTED_MULTIPLICATION_QUADRATIC
+#endif
+
+// Reconstruction
+#define RECONSTRUCTION_TYPE_SCATTER 0
+#define RECONSTRUCTION_TYPE_GATHER	1
+
+#ifndef RECONSTRUCTION_TYPE
+#define RECONSTRUCTION_TYPE RECONSTRUCTION_TYPE_SCATTER
+#endif
+
+#define RADIANCE_PREPROCESS_SCALE_FACTOR 0.1f
+#define RADIANCE_POSTPROCESS_INVERSE_SCALE_FACTOR (1/RADIANCE_PREPROCESS_SCALE_FACTOR);
+
+//--------------------------------------------------------------------------------------------------
+// Util and general texture operations
+//--------------------------------------------------------------------------------------------------
+#define TEXTURE_OPS_MULTIPLY		0
+#define TEXTURE_OPS_DIVIDE			1
+#define TEXTURE_OPS_ADD_CONSTANT	2
+#define TEXTURE_OPS_ACCUMULATE		3
+
+#ifndef TEXTURE_OPS
+#define TEXTURE_OPS TEXTURE_OPS_MULTIPLY
+#endif
+
+#ifndef ACCUMULATE_BY_MASK
+#define ACCUMULATE_BY_MASK 0
+#endif
+
+float Length2(TPixelValue Value)
+{
+#if SOURCE_CHANNEL_COUNT == 1
+	return Value * Value;
+#elif SOURCE_CHANNEL_COUNT == 4
+	return length2(Value.rgb);
+#else
+	return length2(Value);
+#endif
+}
+
+TPixelValue GetImageValue(int2 P, Texture2D<TPixelValue> inImage)
+{
+	return inImage.Load(int3(P, 0));
+}
+
+int2 GetMirroredPosition(int2 P, int2 inTextureSize)
+{
+	int2 TextureSizeMax = inTextureSize - 1;
+	P = abs(TextureSizeMax - abs(P - TextureSizeMax));
+	return P;
+}
+
+//--------------------------------------------------------------------------------------------------
+// General texture operations
+Texture2D<TPixelValue> Source;
+Texture2D<uint> Mask;
+RWTexture2D<float4> RWTarget;
+int2 SourcePosition;
+int2 TargetPosition;
+int2 Size;
+int ForceOperation;
+float4 ConstantValue;
+
+float4 LoadSource(uint2 Position)
+{
+	float4 SourceValue = 0;
+
+#if SOURCE_CHANNEL_COUNT == 1
+	SourceValue.x = Source.Load(uint3(Position, 0));
+#elif SOURCE_CHANNEL_COUNT == 2
+	SourceValue.xy = Source.Load(uint3(Position, 0));
+#elif SOURCE_CHANNEL_COUNT == 3
+	SourceValue.xyz = Source.Load(uint3(Position, 0));
+#elif SOURCE_CHANNEL_COUNT == 4
+	SourceValue = Source.Load(uint3(Position, 0));
+#endif
+	return SourceValue;
+}
+
+[numthreads(THREAD_GROUP_SIZE, THREAD_GROUP_SIZE, 1)]
+void TextureOperationCS(in const uint3 DispatchThreadID : SV_DispatchThreadID)
+{
+	if (any(DispatchThreadID.xy >= (uint2)Size))
+	{
+		return;
+	}
+
+	const uint2 Position = DispatchThreadID.xy;
+	const uint2 ResolvedSourcePosition = max(SourcePosition + Position, 0);
+	const uint2 ResolvedTargetPosition = TargetPosition + Position;
+
+	if (all(ResolvedTargetPosition >= 0))
+	{
+#if TEXTURE_OPS == TEXTURE_OPS_MULTIPLY
+		float4 STexel = LoadSource(ResolvedSourcePosition);
+		RWTarget[ResolvedTargetPosition].rgb *= lerp(1.0f, STexel.rgb, (STexel.rgb != 0) | ForceOperation);
+#elif TEXTURE_OPS == TEXTURE_OPS_DIVIDE
+		float4 STexel = LoadSource(ResolvedSourcePosition);
+		RWTarget[ResolvedTargetPosition].rgb /= lerp(1.0f, STexel.rgb, (STexel.rgb != 0) | ForceOperation);
+#elif TEXTURE_OPS == TEXTURE_OPS_ADD_CONSTANT
+		float3 Value = 0;
+#if ACCUMULATE_BY_MASK
+		uint MaskId = clamp(Mask[ResolvedTargetPosition], 0, 3); // support up to 4 masked add.
+		Value = ConstantValue[MaskId];
+#else
+		Value = ConstantValue.rgb;
+#endif
+		RWTarget[ResolvedTargetPosition].rgb += Value;
+#elif TEXTURE_OPS == TEXTURE_OPS_ACCUMULATE
+		float4 STexel = LoadSource(ResolvedSourcePosition);
+		RWTarget[ResolvedTargetPosition] += STexel;
+#endif
+	}
+}
+
+int2 SourceOffset;
+float4 CopyTexturePS(float4 SvPosition : SV_POSITION) : SV_Target0
+{
+	uint2 Position = (uint2)SvPosition.xy + SourceOffset;
+	return Source.Load(uint3(GetMirroredPosition(Position, TextureSize), 0));
+}
+
+//--------------------------------------------------------------------------------------------------
+// Variance definition and functions
+//
+// Assumption: Variance texture is of type float4 and holds std instead
+//
+#define TImageVariance float4
+
+float GetImageVariance(int2 P, Texture2D inVariance, int inVarianceChannelOffset)
+{
+	float Var = 0;
+#if (IMAGE_VARIANCE_TYPE == IMAGE_VARIANCE_NORMAL) | (IMAGE_VARIANCE_TYPE == IMAGE_VARIANCE_GREYSCALE)
+	// The loaded variance is actually the std
+	float Std = inVariance.Load(int3(P, 0))[inVarianceChannelOffset];
+	Var = Pow2(Std);	
+#elif IMAGE_VARIANCE_TYPE == IMAGE_VARIANCE_COLORED
+	// Assume channel independence and use Luminance() for grey scale
+	// calculation.
+	const float3 Constants = float3(0.09, 0.3481, 0.0121);
+	float3 StdRGB = inVariance.Load(int3(P, 0)).rgb;
+	Var = dot(Pow2(StdRGB), Constants);
+#else
+#error IMAGE_VARIANCE_TYPE not supported
+#endif
+	return Var;
+}
+
+//------------------------------------------------------------------------------------------------------
+// Radiance normalization
+//------------------------------------------------------------------------------------------------------
+
+Texture2D<float4> Normal;
+Texture2D<float4> NormalVariance;
+RWTexture2D<uint> RWMask;
+int2 TextureSize;
+
+[numthreads(THREAD_GROUP_SIZE, THREAD_GROUP_SIZE, 1)]
+void ClassifyPreAlbedoDivideMaskIdCS(in const uint3 DispatchThreadID : SV_DispatchThreadID)
+{
+	if (any(DispatchThreadID.xy >= (uint2)TextureSize))
+	{
+		return;
+	}
+
+	const uint2 Position = DispatchThreadID.xy;
+	float4 NormalValues = 0;
+	NormalValues.rgb = Normal.Load(uint3(Position, 0)).rgb;
+	NormalValues.a = NormalVariance.Load(uint3(Position, 0)).b;
+
+	uint MaskId = 0;
+
+	if (all(NormalValues == 0))
+	{
+		MaskId = 1;
+	}
+
+	RWMask[Position] = MaskId;
+}
+
+Texture2D<float4> Albedo;
+RWTexture2D<float4> RWRadianceVariance;
+
+[numthreads(THREAD_GROUP_SIZE, THREAD_GROUP_SIZE, 1)]
+void NormalizeRadianceVarianceByAlbedoCS(in const uint3 DispatchThreadID : SV_DispatchThreadID)
+{
+	if (any(DispatchThreadID.xy >= (uint2)Size))
+	{
+		return;
+	}
+
+	const uint2 Position = DispatchThreadID.xy;
+	float AlbedoLum = Luminance(Albedo.Load(uint3(Position, 0)).rgb);
+
+	// Since r stores the std instead of variance, need to divide albedo instead of albedo^2.
+	RWRadianceVariance[Position].r /= lerp(1.0f, AlbedoLum, AlbedoLum > 0);
+}
+
+//--------------------------------------------------------------------------------------------------
+// Non-local mean definition and function
+//--------------------------------------------------------------------------------------------------
+// When NONLOCALMEAN_SEPARATE_SOURCE == 1, the NLM calculates the distance from source image to target image
+// instead of on its own.
+
+int PatchSize;
+int PatchDistance;
+float Bandwidth;
+
+Texture2D<TPixelValue> Image;
+Texture2D<TImageVariance> Variance;
+
+#if NONLOCALMEAN_SEPARATE_SOURCE
+Texture2D<TPixelValue> TargetImage;
+Texture2D<TImageVariance> TargetVariance;
+#else
+#define TargetImage Image
+#define TargetVariance Variance
+#endif
+
+int VarianceChannelOffset;
+
+float GetSquaredDistance(int2 P, int2 Q, float VarianceScale = 1.0f)
+{
+	// Use mirrored position. Clamp to boarder has artifacts at border.
+	P = GetMirroredPosition(P, TextureSize); 
+	Q = GetMirroredPosition(Q, TextureSize); 
+
+	TPixelValue Cp = GetImageValue(P, Image);
+	TPixelValue Cq = GetImageValue(Q, TargetImage);
+
+	float VarP = GetImageVariance(P, Variance, VarianceChannelOffset) * VarianceScale;
+	float VarQ = GetImageVariance(Q, TargetVariance, VarianceChannelOffset) * VarianceScale;
+
+	const float Epsilon =  1e-8f; // Small enough to preserve dark areas
+
+	float dpq = (Length2(Cp - Cq) - (VarP + min(VarP, VarQ))) / (Epsilon + Pow2(Bandwidth)*(VarP + VarQ));
+	
+	const float MaxDistance = 10000;
+	dpq = min(dpq, MaxDistance);
+
+	return dpq;
+}
+
+float GetPatchSquaredDistance(int2 P, int2 Q)
+{
+	const float VarianceScale = 1.0f;
+
+	// TODO: z-order for performance improvement?
+	float Sum = 0;
+	for (int y = -PatchSize; y < PatchSize; ++y)
+	{
+		for (int x = -PatchSize; x < PatchSize; ++x)
+		{
+			int2 Pij = P + int2(x, y);
+			int2 Qij = Q + int2(x, y);
+			Sum += GetSquaredDistance(Pij, Qij, VarianceScale);
+		}
+	}
+
+	// ||Pp - Pq||^2
+	float PathSquareDistance = max(0, Sum / Pow2(2 * PatchSize + 1));
+	return PathSquareDistance;
+}
+
+// P is on source image, Q is on target image
+float GetNonLocalMeanWeight(int2 P, int2 Q)
+{
+	//TODO: add a box filter?
+	float PatchpqSqaure = GetPatchSquaredDistance(P,Q);
+
+	float NonLocalMeanWeight = 1.0f;
+#if NONLOCALMEAN_WEIGHT_TYPE == NONLOCALMEAN_PATCHONLY
+	NonLocalMeanWeight = FastExp(-PatchpqSqaure);
+#elif NONLOCALMEAN_WEIGHT_TYPE == NONLOCALMEAN_DISTACNE_PATCH
+	const float SigmaS = 2.0f;//TODO: move to CVar
+	NonLocalMeanWeight = FastExp(-length2(P - Q) / (2 * Pow2(SigmaS))) * FastExp(-PatchpqSqaure);
+#endif
+
+#if PRE_ALBEDO_DIVIDE == PRE_ALBEDO_DIVIDE_DISABLED
+	const float kMinFloat16 = 0.000000059604645f;
+	//without albedo divide, below is better when scattering results.
+	NonLocalMeanWeight = max(kMinFloat16, NonLocalMeanWeight);
+#endif
+	return NonLocalMeanWeight;
+}
+
+RWTexture2D<TPixelValue> DenoisedImage;
+RWBuffer<float> NonLocalMeanWeights;
+int DenoisingChannelCount;
+
+[numthreads(THREAD_GROUP_SIZE, THREAD_GROUP_SIZE, 1)]
+void NonLocalMeanFilteringCS(in const uint3 DispatchThreadID : SV_DispatchThreadID)
+{
+	if (any(DispatchThreadID.xy >= TextureSize))
+	{
+		return;
+	}
+	
+	const uint2 Position = DispatchThreadID.xy;
+	float TotalWeights = 0;
+	TPixelValue Result = 0;
+
+	for (int y = -PatchDistance; y <= PatchDistance; ++y)
+	{
+		for (int x = -PatchDistance; x <= PatchDistance; ++x)
+		{
+			int2 Q = Position + int2(x, y);
+			float Weight = GetNonLocalMeanWeight(Position, Q);
+
+			Result += Weight * GetImageValue(GetMirroredPosition(Q, TextureSize), TargetImage);
+			TotalWeights += Weight;
+		}
+	}
+
+	Result = Result / max(1e-6f, TotalWeights);
+
+#if SOURCE_CHANNEL_COUNT == 1
+	DenoisedImage[DispatchThreadID.xy] = Result;
+#else
+	TPixelValue TargetValue = GetImageValue(GetMirroredPosition(Position, TextureSize), TargetImage);
+	for (int i = 0; i < DenoisingChannelCount; ++i)
+	{
+		TargetValue[i] = Result[i];
+	}
+	DenoisedImage[DispatchThreadID.xy] = TargetValue;
+#endif
+}
+
+// Calculate the non-local mean weights for Region
+int4 Region;
+
+[numthreads(THREAD_GROUP_SIZE, THREAD_GROUP_SIZE, 1)]
+void NonLocalMeanWeightsCS(in const uint3 DispatchThreadID : SV_DispatchThreadID)
+{
+	int2 RegionSize = int2(Region.z - Region.x, Region.w - Region.y);
+	if (any(DispatchThreadID.xy >= RegionSize))
+	{
+		return;
+	}
+	const uint2 Position = DispatchThreadID.xy + Region.xy;
+	const int PatchDiameter = 2 * PatchDistance + 1;
+	const int PixelOffset = Pow2(PatchDiameter);
+	const int BufferOffset = (DispatchThreadID.x + DispatchThreadID.y * RegionSize.x) * PixelOffset;
+
+	for (int i = 0; i < PixelOffset; ++i)
+	{
+		int2 Offset = int2(i % PatchDiameter, i / PatchDiameter) - int2(PatchDistance, PatchDistance);
+		int2 Q = Position + Offset;
+		NonLocalMeanWeights[BufferOffset + i] = GetNonLocalMeanWeight(Position, Q);
+	}
+}
+
+//------------------------------------------------------------------------------------------------------
+// Collaborative filtering
+//	1. Tiling
+//------------------------------------------------------------------------------------------------------
+
+RWBuffer<float> Dest;
+int CopyChannelOffset;
+int CopyChannelCount;
+int BufferChannelOffset;
+int BufferChannelSize;
+int4 CopyRegion;
+
+[numthreads(THREAD_GROUP_SIZE, THREAD_GROUP_SIZE, 1)]
+void CopyTextureToBufferCS(in const uint3 DispatchThreadID : SV_DispatchThreadID)
+{
+	uint2 CopyRegionSize = uint2(CopyRegion.z - CopyRegion.x, CopyRegion.w - CopyRegion.y);
+	if (any(DispatchThreadID.xy >= CopyRegionSize))
+	{
+		return;
+	}
+
+	const uint2 Position = DispatchThreadID.xy + CopyRegion.xy;
+	const int BufferOffset = (DispatchThreadID.x + DispatchThreadID.y * CopyRegionSize.x) * BufferChannelSize + BufferChannelOffset;
+
+	// If copy region is outside of the texture region, get a mirror.
+	TPixelValue  PixelValue = Source.Load(uint3(GetMirroredPosition(Position, TextureSize), 0));
+
+	const int EndCopyChannelOffset = CopyChannelOffset + CopyChannelCount;
+
+#if SOURCE_CHANNEL_COUNT == 1
+	Dest[BufferOffset] = PixelValue;
+#else
+	for (int i = CopyChannelOffset; i < EndCopyChannelOffset; ++i)
+	{
+		Dest[BufferOffset + (i - CopyChannelOffset)] = PixelValue[i];
+	}
+#endif
+}
+
+RWTexture2D<float4> RWSource;
+
+[numthreads(THREAD_GROUP_SIZE, THREAD_GROUP_SIZE, 1)]
+void NormalizeTextureCS(in const uint3 DispatchThreadID : SV_DispatchThreadID)
+{
+	if (any(DispatchThreadID.xy >= TextureSize))
+	{
+		return;
+	}
+	
+	const uint2 Position = DispatchThreadID.xy;
+	float4 Texel = RWSource[Position];
+	
+	if (Texel.w == 0)
+	{
+		Texel.w = 1;
+	}
+
+	RWSource[Position] = Texel / Texel.w;
+}
+
+//------------------------------------------------------------------------------------------------------
+//	2. Weighted Least-square solver
+
+Buffer<float> X;
+Buffer<float> W;
+Buffer<float> Y;
+
+RWBuffer<float> Result;
+
+int2 XDim;
+int WDim;
+int NumOfTemporalFrames;
+int NumOfWeigthsPerPixelPerFrame;
+
+#if WEIGHTED_MULTIPLICATION_TYPE == WEIGHTED_MULTIPLICATION_GENERALIZED
+int2 YDim;
+#else
+#define YDim XDim
+#endif
+
+// X, Y is stored per pixel, in each pixel, it stores the pixel footage for each temporal frame
+// Pixel layout example with two data per temporal frame
+// P1                     ...   Pn
+// |F0 F1|F0 F1|F0 F1|
+//   T0     T1   T2
+// 
+// W stores one weight per pixel, when a frame stored, the next frame is stored
+// P1 ... Pn | P1 ... Pn|
+// T0          T1        ...
+
+float GetX(uint2 Position, uint f)
+{
+	const uint Width = TextureSize.x + 2 * PatchDistance;
+	const uint F = XDim.y;
+
+	const uint Index = (Position.y * Width + Position.x) * (F * NumOfTemporalFrames) + f;
+	return X[Index];
+}
+
+float GetX(uint2 Position, uint Db, uint FrameIndex)
+{
+	const uint F = XDim.y;
+	return GetX(Position, FrameIndex * F + Db);
+}
+
+float GetY(int2 Position, int f)
+{
+#if WEIGHTED_MULTIPLICATION_TYPE == WEIGHTED_MULTIPLICATION_GENERALIZED
+	const float Width = TextureSize.x + 2 * PatchDistance;
+	const float F = YDim.y;
+
+	const int Index = (Position.y * Width + Position.x) * (F * NumOfTemporalFrames) + f;
+	return Y[Index];
+#else
+	return GetX(Position, f);
+#endif
+}
+
+float GetW(int2 Position, int WeightIndex, int FrameIndex)
+{
+	const int NumberOfPixels = TextureSize.x * TextureSize.y;
+	const int Index = NumOfWeigthsPerPixelPerFrame *(NumberOfPixels * FrameIndex + (Position.y * TextureSize.x + Position.x)) + WeightIndex;
+	return W[Index];
+}
+
+int2 GetModifiedXyYyDimension()
+{
+	const int ModifiedXyDim = XDim.y + APPEND_CONSTANT_DIMENSION_TO_X;
+#if WEIGHTED_MULTIPLICATION_TYPE == WEIGHTED_MULTIPLICATION_GENERALIZED
+	const int ModifiedYyDim = YDim.y;
+#else
+	const int ModifiedYyDim = YDim.y + APPEND_CONSTANT_DIMENSION_TO_X;
+#endif
+
+	return int2(ModifiedXyDim, ModifiedYyDim);
+}
+
+void WriteResult(int2 Position, int f, float Value)
+{
+	const int2 ModifiedXyYyDim = GetModifiedXyYyDimension();
+	const int Index = (Position.y * TextureSize.x + Position.x) * (ModifiedXyYyDim.x * ModifiedXyYyDim.y) + f;
+	Result[Index] = Value;
+}
+
+#include "/Engine/Private/DoubleFloat.ush"
+
+int SamplingStep;
+
+void SolveMatrixMultiplicationPerMatrixItem(int2 Position, int2 ModifiedXyYyDim, int PatchWidth, int SingleFrameDataCount)
+{
+	for (int i = 0; i < ModifiedXyYyDim.x; ++i)
+	{
+#if WEIGHTED_MULTIPLICATION_TYPE == WEIGHTED_MULTIPLICATION_QUADRATIC
+		for (int k = i; k < ModifiedXyYyDim.y; ++k)
+#else
+		for (int k = 0; k < ModifiedXyYyDim.y; ++k)
+#endif
+		{
+			FDFScalar R = (FDFScalar)0;
+			for (int j = 0; j < XDim.x; ++j)
+			{
+
+				int IndexWithinFrame = j % SingleFrameDataCount;
+				int2 Offset = int2(IndexWithinFrame % PatchWidth, IndexWithinFrame / PatchWidth);
+				int frame_index = j / SingleFrameDataCount;
+
+				float w = GetW(Position, IndexWithinFrame, frame_index); // get w, cache ???, 19*19*T
+
+				int2 TartPosition = Position + Offset;
+
+#if APPEND_CONSTANT_DIMENSION_TO_X
+				float x = 1.0f;
+				float y = 1.0f;
+
+				BRANCH
+					if (i < XDim.y)
+					{
+						x = GetX(TartPosition, XDim.y * frame_index + i);
+					}
+
+				BRANCH
+					if (k < YDim.y)
+					{
+						y = GetY(TartPosition, YDim.y * frame_index + k);
+					}
+#else
+				float x = GetX(TartPosition, XDim.y * frame_index + i); // get x from the ith row X^T
+				float y = GetY(TartPosition, YDim.y * frame_index + k); // get y from the kth column of Y
+#endif
+				R = DFAdd(R, DFMultiply(DFTwoSum(x, 0.0f), DFMultiply(DFTwoSum(sqrt(w), 0.0f), DFTwoSum(y, 0.0f))));
+			}
+
+			float RResolved = DFDemote(R);
+			WriteResult(Position, i * ModifiedXyYyDim.y + k, RResolved);
+
+#if WEIGHTED_MULTIPLICATION_TYPE == WEIGHTED_MULTIPLICATION_QUADRATIC
+			WriteResult(Position, k * ModifiedXyYyDim.y + i, RResolved);
+#endif
+		}
+	}
+}
+
+struct FMatrixCache
+{
+#if WEIGHTED_MULTIPLICATION_TYPE == WEIGHTED_MULTIPLICATION_QUADRATIC
+	float Data[((NUM_FEATURE + APPEND_CONSTANT_DIMENSION_TO_X) * (NUM_FEATURE + APPEND_CONSTANT_DIMENSION_TO_X) + 1) / 2];
+#elif WEIGHTED_MULTIPLICATION_TYPE == WEIGHTED_MULTIPLICATION_GENERALIZED
+	float Data[(NUM_FEATURE + APPEND_CONSTANT_DIMENSION_TO_X) * 3];
+#endif
+};
+
+int SourceFrameIndex;
+
+void MatrixMultiplicationSmallMatrixOptimization(int2 Position, int2 ModifiedXyYyDim, int PatchWidth, int SingleFrameDataCount)
+{	
+	FMatrixCache MatrixCache = (FMatrixCache)0;
+
+	int NumOfFrames = XDim.x / SingleFrameDataCount;
+	int Index = 0;
+
+#if USE_SAMPLING_STEP
+	#if SAMPLING_TYPE == SAMPLING_TYPE_INCREMENTAL
+	//This sampling might lead to incorrect DOF when there is ambiguity using feature to predict
+	//non-DOF background in previous/future frames, and DOF foreground that is the current frame.
+	//we might need to apply regression for each frame instead of using combined frames.
+	for (int step = 0; step < XDim.x; step+=SamplingStep)
+	{
+		int IndexWithinFrame = step / NumOfFrames;
+		int2 Offset = int2(IndexWithinFrame % PatchWidth, IndexWithinFrame / PatchWidth);
+		int2 TartPosition = Position + Offset;
+
+		int t = step % NumOfFrames;
+	#elif SAMPLING_TYPE == SAMPLING_TYPE_FOCUS_CURRENT_FRAME
+		// focus sampling more to the current frames.
+	const int SourceFrame = SourceFrameIndex;
+
+	int PreviousSamplingFrame = NumOfFrames - 1;
+	int IndexWithinFrame = 0;
+	for (int step = 0; step < XDim.x && IndexWithinFrame < SingleFrameDataCount-1; step += SamplingStep)
+	{
+		int t = step % NumOfFrames;
+		if (t < PreviousSamplingFrame && NumOfFrames > 1)
+		{
+			PreviousSamplingFrame = t;
+			t = SourceFrame;
+			IndexWithinFrame += 1;
+		}
+		else
+		{
+			int Tmp = t;
+			t = SourceFrame;
+			PreviousSamplingFrame = Tmp;
+		}
+
+		int2 Offset = int2(IndexWithinFrame % PatchWidth, IndexWithinFrame / PatchWidth);
+		int2 TartPosition = Position + Offset;
+	#endif
+#else
+	for (int j = 0; j < SingleFrameDataCount; ++j)
+	{
+		int IndexWithinFrame = j;
+		int2 Offset = int2(IndexWithinFrame % PatchWidth, IndexWithinFrame / PatchWidth);
+		int2 TartPosition = Position + Offset;
+
+		for (int t = 0; t < NumOfFrames; ++t)
+#endif
+		{
+			int frame_index = t;
+			float sqrtw = sqrt(GetW(Position, IndexWithinFrame, frame_index)); // get w.
+			
+			Index = 0;
+			for (int i = 0; i < ModifiedXyYyDim.x; ++i)
+			{
+#if WEIGHTED_MULTIPLICATION_TYPE == WEIGHTED_MULTIPLICATION_QUADRATIC
+				for (int k = i; k < ModifiedXyYyDim.y; ++k)
+#else
+				for (int k = 0; k < ModifiedXyYyDim.y; ++k)
+#endif
+				{
+
+#if APPEND_CONSTANT_DIMENSION_TO_X
+					float x = 1.0f;
+					float y = 1.0f;
+
+					BRANCH
+						if (i < XDim.y)
+						{
+							x = GetX(TartPosition, XDim.y * frame_index + i);
+						}
+
+					BRANCH
+						if (k < YDim.y)
+						{
+							y = GetY(TartPosition, YDim.y * frame_index + k);
+						}
+#else
+					float x = GetX(TartPosition, XDim.y * frame_index + i); // get x from the ith row X^T
+					float y = GetY(TartPosition, YDim.y * frame_index + k); // get y from the kth column of Y
+#endif
+					MatrixCache.Data[Index] += x * sqrtw * y;
+					++Index;
+				}
+			}
+		}
+	}
+
+	Index = 0;
+	for (int i = 0; i < ModifiedXyYyDim.x; ++i)
+	{
+#if WEIGHTED_MULTIPLICATION_TYPE == WEIGHTED_MULTIPLICATION_QUADRATIC
+		for (int k = i; k < ModifiedXyYyDim.y; ++k)
+#else
+		for (int k = 0; k < ModifiedXyYyDim.y; ++k)
+#endif
+		{
+			float RResolved = MatrixCache.Data[Index++];
+			WriteResult(Position, i * ModifiedXyYyDim.y + k, RResolved);
+
+			#if (WEIGHTED_MULTIPLICATION_TYPE == WEIGHTED_MULTIPLICATION_QUADRATIC)
+				WriteResult(Position, k * ModifiedXyYyDim.y + i, RResolved);
+			#endif
+		}
+	}
+}
+
+[numthreads(THREAD_GROUP_SIZE, THREAD_GROUP_SIZE, 1)]
+void InPlaceBatchedMatrixMultiplicationCS(in const uint3 DispatchThreadID : SV_DispatchThreadID)
+{
+	if (any(DispatchThreadID.xy >= TextureSize))
+	{
+		return;
+	}
+	int2 Position = DispatchThreadID.xy;
+	
+	int PatchWidth = 2 * PatchDistance + 1;
+	int SingleFrameDataCount = Pow2(2 * PatchDistance + 1);//TODO
+	
+	//======================================================================================================
+	// For each pixel, solve the matrix X^T W * Y or X^T W * X
+	// asymptotic complexity O(n^3).
+	// Actual: FxAxN = 21*N  for X^T W Y, where N = 19*19*T, where T = 1,3,5
+	//         FxNxN = 21*N^2 for X^T W X
+	// The time spent for solving the quadratic form increases quadratically, 1, 9, 25.
+	// Improvement: 
+	// 1. Based on "one in ten/twenty rules", we might not need that many Ns.
+	// E.g., if 7x3 weights needed to be predicted, we might shrink 19*19*5 to 19*19*(1~2) observations that is 
+	// distributed over 5 frames.
+	const int2 ModifiedXyYyDim = GetModifiedXyYyDimension();
+
+#if SMALL_MATRIX_OPTIMIZE
+	MatrixMultiplicationSmallMatrixOptimization(Position, ModifiedXyYyDim, PatchWidth, SingleFrameDataCount);
+#else
+	SolveMatrixMultiplicationPerMatrixItem(Position, ModifiedXyYyDim, PatchWidth, SingleFrameDataCount);
+#endif
+}
+
+#define MATRIX_DIM NUM_FEATURE
+#include "NFORRegression.ush"
+
+[numthreads(THREAD_GROUP_SIZE, THREAD_GROUP_SIZE, 1)]
+void LinearSolverCS(in const uint3 DispatchThreadID : SV_DispatchThreadID)
+{
+	const int MatrixIndex = DispatchThreadID.x + DispatchThreadID.y * NumOfElementsPerRow;
+	if (MatrixIndex >= NumOfElements)
+	{
+		return;
+	}
+	const int BSize = BDim.x * BDim.y;
+	const int BOffset = MatrixIndex * BSize;
+	const int ASize = ADim.x * ADim.y;
+	const int AOffset = MatrixIndex * ASize;
+
+	//QR_Solve(AOffset, ASize, BOffset);
+	NewtonSchulzIterationDouble(AOffset, ASize, BOffset);
+}
+
+//===============================================================================================
+// Reconstruct spatial temporal images
+int FrameIndex;
+float AlbedoOffset;
+RWTexture2D<float4> RWReconstruction;
+RWStructuredBuffer<uint4> RWReconstructBuffer;
+RWTexture2D<UlongType> RWReconstructBuffer64;
+
+float GetB(int2 Position, int f)
+{
+	const int Index = (Position.y * TextureSize.x + Position.x) * (BDim.x*BDim.y) + f;
+	return B[Index];
+}
+
+uint4 EncodeFloat4(float4 Value)
+{
+	uint4 fractions = frac(max(Value,0.0f))*0xffff;
+	uint4 Integers = uint4(max(Value, 0.0f));
+	return Integers<<16 | (fractions & 0xffff);
+}
+
+float4 DecodeFloat4(uint4 Value)
+{
+	float4 Decoded = float4(Value>>16) + float4(Value & 0xffff)/65536;
+}
+
+void SaveFloat4ToBufferOrTexture(uint2 Position, float4 Value, uint NumOfScattering)
+{
+#if COMPILER_SUPPORTS_UINT64_IMAGE_ATOMICS
+	float4 XBWFraction = frac(max(Value, 0));
+	uint4 XBWInteger = Value - XBWFraction;
+	XBWInteger = min(XBWInteger, 0xffffffff / NumOfScattering); // Avoid overflow.
+
+	uint4 XBWFractionAsInteger = uint4(XBWFraction * 0xffffffff);
+	{
+		const UlongType R = PackUlongType(uint2(XBWFractionAsInteger.r, XBWInteger.r));
+		ImageInterlockedAddUInt64(RWReconstructBuffer64, uint2(Position.x * 4 + 0, Position.y), R);
+
+		const UlongType G = PackUlongType(uint2(XBWFractionAsInteger.g, XBWInteger.g));
+		ImageInterlockedAddUInt64(RWReconstructBuffer64, uint2(Position.x * 4 + 1, Position.y), G);
+
+		const UlongType B = PackUlongType(uint2(XBWFractionAsInteger.b, XBWInteger.b));
+		ImageInterlockedAddUInt64(RWReconstructBuffer64, uint2(Position.x * 4 + 2, Position.y), B);
+
+		const UlongType A = PackUlongType(uint2(XBWFractionAsInteger.a, XBWInteger.a));
+		ImageInterlockedAddUInt64(RWReconstructBuffer64, uint2(Position.x * 4 + 3, Position.y), A);
+	}
+
+#else
+	int IntermediateBufferIndex = Position.x + Position.y * (TextureSize.x + 2 * PatchDistance);
+	uint4 XBWEncoded = EncodeFloat4(XBW);
+	float4 OldValue = 0;
+	InterlockedAdd(RWReconstructBuffer[IntermediateBufferIndex].x, XBWEncoded.x, OldValue.x);
+	InterlockedAdd(RWReconstructBuffer[IntermediateBufferIndex].y, XBWEncoded.y, OldValue.y);
+	InterlockedAdd(RWReconstructBuffer[IntermediateBufferIndex].z, XBWEncoded.z, OldValue.z);
+	InterlockedAdd(RWReconstructBuffer[IntermediateBufferIndex].w, XBWEncoded.w, OldValue.w);
+
+#endif //COMPILER_SUPPORTS_UINT64_IMAGE_ATOMICS
+}
+
+[numthreads(THREAD_GROUP_SIZE, THREAD_GROUP_SIZE, 1)]
+void ReconstructSpatialTemporalImageCS(in const uint3 DispatchThreadID : SV_DispatchThreadID)
+{
+	if (any(DispatchThreadID.xy >= TextureSize))
+	{
+		return;
+	}
+	int2 Position = DispatchThreadID.xy;
+
+	const int NumOfFeatures = XDim.y;
+	const int PatchSideLength = 2 * PatchDistance + 1;
+
+	// Get the weights for all the pixel's patch's reconstruction
+	const int BWeightStride = BDim.x * BDim.y;
+
+	float3 Weights[NUM_FEATURE]; // 6x3 matrix.
+	
+	// checkf(NUM_FEATURE == XDim.y)
+
+	UNROLL
+	for (int i = 0; i < NUM_FEATURE; ++i)
+	{
+		float3 weight = 0;
+
+		weight.x = GetB(Position, i * BDim.y + 0);
+		weight.y = GetB(Position, i * BDim.y + 1);
+		weight.z = GetB(Position, i * BDim.y + 2);
+		Weights[i] = weight;
+	}
+
+	const int WeightStride = Pow2(PatchSideLength);
+	const int XStride = XDim.y * WeightStride;
+	float4 GatheredValue = 0;
+
+#if PRE_ALBEDO_DIVIDE == PRE_ALBEDO_DIVIDE_EACH
+	float3 Albedo = 0;
+#endif
+
+	// For center frame, reconstruct and scatter the result
+	for (int y = 0; y < PatchSideLength; ++y)
+	{
+		for (int x = 0; x < PatchSideLength; ++x)
+		{
+			int WeightIndex1D = (x + y * PatchSideLength);
+			float w = GetW(Position, WeightIndex1D, FrameIndex);
+
+			uint2 TargetPosition = Position + int2(x, y);
+			float3 Denoised = 0;
+			
+			for (int i = 0; i < NUM_FEATURE; ++i)
+			{
+				float FeatureI = 1.0f; 
+				if (i < NumOfFeatures)
+				{
+					FeatureI = GetX(TargetPosition, i, FrameIndex);
+
+#if PRE_ALBEDO_DIVIDE == PRE_ALBEDO_DIVIDE_EACH
+						//Assume the first three channels are albedo. TODO: use passed in index
+						if (i < 3)
+						{
+							Albedo[i] = FeatureI;
+						}
+#endif
+				}
+				Denoised += FeatureI * Weights[i];
+			}
+
+#if RECONSTRUCTION_TYPE == RECONSTRUCTION_TYPE_SCATTER
+
+#if PRE_ALBEDO_DIVIDE == PRE_ALBEDO_DIVIDE_EACH
+			Denoised *= RADIANCE_PREPROCESS_SCALE_FACTOR * Albedo;
+#elif PRE_ALBEDO_DIVIDE == PRE_ALBEDO_DIVIDE_FINAL
+			Denoised *= RADIANCE_PREPROCESS_SCALE_FACTOR;
+#endif
+			// Scatter the result instead of gathering for the current frame to get high frequency results (better reconstruction).
+			float4 XBW = float4(Denoised * w, w);
+			SaveFloat4ToBufferOrTexture(TargetPosition, XBW, WeightStride);
+
+#elif RECONSTRUCTION_TYPE == RECONSTRUCTION_TYPE_GATHER
+
+#if PRE_ALBEDO_DIVIDE == PRE_ALBEDO_DIVIDE_EACH
+			Denoised *= Albedo;
+#endif
+			GatheredValue += float4(Denoised * w, w);
+#endif
+		}
+	}
+#if RECONSTRUCTION_TYPE == RECONSTRUCTION_TYPE_GATHER
+	// For other frames, after reconstruction, perform reconstruction with the weights.
+	RWReconstruction[Position + PatchDistance] += GatheredValue;
+#endif
+
+}
+
+StructuredBuffer<uint4> StructuredBufferSource;
+Texture2D<UlongType> ReconstructBuffer64;
+
+float DecodeUint64ToFloat(Texture2D<UlongType> inTexture, uint2 Position)
+{
+	const uint2 REncoded = UnpackUlongType(inTexture[Position]);
+	return REncoded.y + float(REncoded.x) / 0xffffffff;
+}
+
+float4 DecodeFloat4FromBufferOrTexture(uint2 Position)
+{
+	float4 DecodedValue = 0;
+#if COMPILER_SUPPORTS_UINT64_IMAGE_ATOMICS
+	float R = DecodeUint64ToFloat(ReconstructBuffer64, uint2(Position.x * 4 + 0, Position.y));
+	float G = DecodeUint64ToFloat(ReconstructBuffer64, uint2(Position.x * 4 + 1, Position.y));
+	float B = DecodeUint64ToFloat(ReconstructBuffer64, uint2(Position.x * 4 + 2, Position.y));
+	float A = DecodeUint64ToFloat(ReconstructBuffer64, uint2(Position.x * 4 + 3, Position.y));
+
+	DecodedValue = float4(R,G,B,A);
+
+#if PRE_ALBEDO_DIVIDE != PRE_ALBEDO_DIVIDE_DISABLED
+	DecodedValue.rgb *= RADIANCE_POSTPROCESS_INVERSE_SCALE_FACTOR;
+#endif
+
+#else
+	int BufferIndex = Position.x + Position.y * TextureSize.x;
+	DecodedValue = DecodeFloat4(StructuredBufferSource[BufferIndex]);
+#endif
+
+	// There would be cases where a is 0 due to w being too small and encoded as 0.
+	// Without rejecting this sample, could see borders artifacts around each tile.
+	DecodedValue = lerp(0, DecodedValue, DecodedValue.a != 0);
+
+	return DecodedValue;
+}
+
+[numthreads(THREAD_GROUP_SIZE, THREAD_GROUP_SIZE, 1)]
+void AccumulateBufferToTextureCS(in const uint3 DispatchThreadID : SV_DispatchThreadID)
+{
+	if (any(DispatchThreadID.xy >= TextureSize))
+	{
+		return;
+	}
+
+	uint2 Position = DispatchThreadID.xy;
+	float4 Color = DecodeFloat4FromBufferOrTexture(Position);
+	RWTarget[DispatchThreadID.xy] += Color;
+}
+
+//------------------------------------------------------------------------------------------------------
+//	Bandwidth selection
+//------------------------------------------------------------------------------------------------------
+
+Texture2D<TPixelValue> FilteredImage;
+RWTexture2D<float> MSE;
+
+[numthreads(THREAD_GROUP_SIZE, THREAD_GROUP_SIZE, 1)]
+void MSEEstimationCS(in const uint3 DispatchThreadID : SV_DispatchThreadID)
+{
+	if (any(DispatchThreadID.xy >= TextureSize))
+	{
+		return;
+	}
+
+	const uint2 Position = DispatchThreadID.xy;
+	float VarC = GetImageVariance(Position, Variance, VarianceChannelOffset);
+	TPixelValue Color = GetImageValue(Position, Image);
+	TPixelValue FilteredColor = GetImageValue(Position, FilteredImage);
+
+	// E[(F-C)^2-Var_C] = Bias^2_F + Var_F - 2Cov(C,F)
+	// Assume there is low correlation between F and C, and low bias of F
+	MSE[Position] = Length2(FilteredColor - Color) - VarC;
+}
+
+Texture2D<float> FilteredMSEs_0;
+Texture2D<float> FilteredMSEs_1;
+RWTexture2D<float> RWSelectionMap;
+
+[numthreads(THREAD_GROUP_SIZE, THREAD_GROUP_SIZE, 1)]
+void GenerateSelectionMapCS(in const uint3 DispatchThreadID : SV_DispatchThreadID)
+{
+	if (any(DispatchThreadID.xy >= TextureSize))
+	{
+		return;
+	}
+
+	const uint2 Position = DispatchThreadID.xy;
+	float MSE0 = FilteredMSEs_0.Load(int3(Position, 0));
+	float MSE1 = FilteredMSEs_1.Load(int3(Position, 0));
+	RWSelectionMap[Position] = MSE0 < MSE1 ? 0.0f : 1.0f;
+}
+
+Texture2D<float4> FilteredImages_0;
+Texture2D<float4> FilteredImages_1;
+Texture2D<float>  SelectionMap;
+RWTexture2D<float4> RWFilteredImage;
+
+[numthreads(THREAD_GROUP_SIZE, THREAD_GROUP_SIZE, 1)]
+void CombineFilteredImageCS(in const uint3 DispatchThreadID : SV_DispatchThreadID)
+{
+	if (any(DispatchThreadID.xy >= TextureSize))
+	{
+		return;
+	}
+
+	const uint2 Position = DispatchThreadID.xy;
+	float4 Color0 = FilteredImages_0.Load(int3(Position, 0));
+	float4 Color1 = FilteredImages_1.Load(int3(Position, 0));
+	float  W = SelectionMap.Load(int3(Position, 0));
+	RWFilteredImage[Position] = lerp(Color0, Color1, W);
+}
diff --git a/Engine/Plugins/Experimental/NFORDenoise/Source/NFORDenoise/Shaders/NFORRegression.ush b/Engine/Plugins/Experimental/NFORDenoise/Source/NFORDenoise/Shaders/NFORRegression.ush
new file mode 100644
--- /dev/null
+++ b/Engine/Plugins/Experimental/NFORDenoise/Source/NFORDenoise/Shaders/NFORRegression.ush
@@ -0,0 +1,503 @@
+// Copyright Epic Games, Inc. All Rights Reserved.
+
+/*=============================================================================
+	NFORRegression.ush
+=============================================================================*/
+
+// TODO: Improve matrix multiplication with Strassem algorithm.
+// TODO: Improve the robustness of QR decomposition.
+// TODO: Find best iteration number for Newton Schulz Iteration Method.
+// TODO: Unroll with constant dimensions.
+// TODO: Merge duplicate shader code.
+
+#pragma once
+
+#ifndef MATRIX_DIM
+#define MATRIX_DIM 8
+#endif
+
+#define FEATURE_MATRIX_SIZE ((MATRIX_DIM)*(MATRIX_DIM))
+
+#define QR_DECOMPOSITION_TYPE_DEFAULT		0
+#define QR_DECOMPOSITION_TYPE_COLUMN_PIVOT	1
+
+#define QR_DECOMPOSITION_TYPE QR_DECOMPOSITION_TYPE_DEFAULT
+
+
+Buffer<float> A;
+Buffer<float> B;
+
+int2 ADim;
+int2 BDim;
+int NumOfElements;
+int NumOfElementsPerRow;
+
+struct FNFORVector
+{
+	FDFScalar V[MATRIX_DIM];
+};
+
+struct FNFORMatrix
+{
+	FDFScalar M[FEATURE_MATRIX_SIZE];
+};
+
+void Identity(out FNFORMatrix Matrix)
+{
+	for (int i = 0; i < MATRIX_DIM; ++i)
+	{
+		for (int j = 0; j < MATRIX_DIM; ++j)
+		{
+			Matrix.M[i * MATRIX_DIM + j] = DFTwoSum(lerp(0.0f, 1.0f, i == j), 0.0f);
+		}
+	}
+}
+
+FNFORMatrix Copy(in FNFORMatrix Matrix)
+{
+	FNFORMatrix Out;
+
+	for (int i = 0; i < FEATURE_MATRIX_SIZE; ++i)
+	{
+		Out.M[i] = Matrix.M[i];
+	}
+
+	return Out;
+}
+
+FNFORMatrix NFORMatrixMultiply(in FNFORMatrix A, in FNFORMatrix B)
+{
+	FNFORMatrix Ret;
+	for (int i = 0; i < MATRIX_DIM; ++i)
+	{
+		for (int k = 0; k < MATRIX_DIM; ++k)
+		{
+			FDFScalar Sum = (FDFScalar)0;
+			for (int j = 0; j < MATRIX_DIM; ++j)
+			{
+				const FDFScalar DFMultiplyScalar = DFMultiply(A.M[i * MATRIX_DIM + j], B.M[j * MATRIX_DIM + k]);
+				Sum = DFAdd(Sum, DFMultiplyScalar);
+			}
+			Ret.M[i * MATRIX_DIM + k] = Sum;
+		}
+	}
+
+	return Ret;
+}
+
+FDFScalar GetNFORVectorNorm2(in FDFScalar V[MATRIX_DIM], int N)
+{
+	FDFScalar L2Norm = (FDFScalar)0;
+
+	for (int i = 0; i < N; ++i)
+	{
+		L2Norm = DFAdd(L2Norm, DFMultiply(V[i], V[i]));
+	}
+
+	L2Norm = DFSqrt(L2Norm);
+
+	return L2Norm;
+}
+
+void SwapIntVector(inout int V[MATRIX_DIM], int i, int j)
+{
+	int Tmp = V[i];
+	V[i] = V[j];
+	V[j] = Tmp;
+}
+
+void SwampColumn(inout FNFORMatrix Matrix, int i, int j)
+{
+	FDFScalar Tmp;
+	for (int row = 0; row < MATRIX_DIM; ++row)
+	{
+		Tmp = Matrix.M[row * MATRIX_DIM + i];
+		Matrix.M[row * MATRIX_DIM + i] = Matrix.M[row * MATRIX_DIM + j];
+		Matrix.M[row * MATRIX_DIM + j] = Tmp;
+	}
+}
+
+FNFORMatrix SelectSubMatrix(in FNFORMatrix Matrix, int StartRowColumnIndex)
+{
+	FNFORMatrix Ret;
+	for (int j = StartRowColumnIndex; j < MATRIX_DIM; ++j)
+	{
+		for (int k = StartRowColumnIndex; k < MATRIX_DIM; ++k)
+		{
+			Ret.M[(j - StartRowColumnIndex) * MATRIX_DIM + (k - StartRowColumnIndex)] = Matrix.M[j * MATRIX_DIM + k];
+		}
+	}
+	return Ret;
+}
+
+int GetPivot(in FNFORMatrix Matrix, int N)
+{
+	FNFORVector Vector;
+
+	FDFScalar MaxNorm = DFTwoSum(-1.0f, 0.0f);
+	int MaxIndex = 0;
+
+	for (int column = 0; column < N; ++column)
+	{
+		for (int row = 0; row < N; ++row)
+		{
+			Vector.V[row] = Matrix.M[row * MATRIX_DIM + column];
+		}
+		FDFScalar L2Norm = GetNFORVectorNorm2(Vector.V, N);
+		if (DFGreater(L2Norm, MaxNorm))
+		{
+			MaxNorm = L2Norm;
+			MaxIndex = column;
+		}
+	}
+
+	return MaxIndex + (MATRIX_DIM - N);
+}
+
+void HouseHoldReflection(in FDFScalar A[MATRIX_DIM], out FNFORMatrix O, int N)
+{
+	FDFScalar V[MATRIX_DIM];
+
+	for (int i = 0; i < N; ++i)
+	{
+		V[i] = A[i];
+	}
+
+	FDFScalar L2Norm = GetNFORVectorNorm2(V, N);
+
+	V[0] = DFAdd(V[0], DFMultiply(DFSign(A[0]), L2Norm));
+
+	L2Norm = GetNFORVectorNorm2(V, N);
+
+	FDFScalar InvL2Norm = DFDivide(1.0f, L2Norm);
+	for (int i = 0; i < N; ++i)
+	{
+		V[i] = DFMultiply(V[i], InvL2Norm);
+	}
+
+	//I - 2 * vv
+	for (int i = 0; i < N; ++i)
+	{
+		for (int j = 0; j < N; ++j)
+		{
+			O.M[i * MATRIX_DIM + j] = DFSubtract(lerp(0, 1, i == j), DFMultiply(2, DFMultiply(V[i], V[j])));
+		}
+	}
+}
+
+// Calculate AX=B with (Household) QR decomposition and back substitution.
+// TODO: Reduce the memory used in calculation.
+// TODO: Since A matrix is symmetric, try Choleskey Decomposition.
+void QRDecomposition(in FNFORMatrix A, out FNFORMatrix Q, out FNFORMatrix R
+#if QR_DECOMPOSITION_TYPE == QR_DECOMPOSITION_TYPE_COLUMN_PIVOT
+	, out int PivotIndices[MATRIX_DIM]
+#endif
+)
+{
+	Identity(Q);
+	R = Copy(A);
+
+	FNFORMatrix H;
+	FNFORMatrix Reflection;
+	for (int i = 0; i < MATRIX_DIM; ++i)
+	{
+		// Select column with the largest norm as pivot.
+#if QR_DECOMPOSITION_TYPE == QR_DECOMPOSITION_TYPE_COLUMN_PIVOT
+		H = SelectSubMatrix(R, i);
+		int Pivot = GetPivot(H, MATRIX_DIM - i);
+		SwapIntVector(PivotIndices, i, Pivot);
+		SwampColumn(R, i, Pivot);
+#endif
+		// Apply Householder reflection.
+		Identity(H);
+
+		FDFScalar Selection[MATRIX_DIM];
+		for (int j = i; j < MATRIX_DIM; ++j)
+		{
+			Selection[j - i] = R.M[j * MATRIX_DIM + i];
+		}
+
+		HouseHoldReflection(Selection, Reflection, MATRIX_DIM - i);
+
+		for (int j = i; j < MATRIX_DIM; ++j)
+		{
+			for (int k = i; k < MATRIX_DIM; ++k)
+			{
+				H.M[j * MATRIX_DIM + k] = Reflection.M[(j - i) * MATRIX_DIM + (k - i)];
+			}
+		}
+
+		Q = NFORMatrixMultiply(Q, H);
+		R = NFORMatrixMultiply(H, R);
+	}
+}
+
+void QR_Solve(int AOffset, int ASize, int BOffset)
+{
+	FNFORMatrix NA;
+	FNFORMatrix NQ;
+	FNFORMatrix NR;
+
+#if QR_DECOMPOSITION_TYPE == QR_DECOMPOSITION_TYPE_COLUMN_PIVOT
+	int PivotIndices[MATRIX_DIM];
+
+	for (int i = 0; i < MATRIX_DIM; ++i)
+	{
+		PivotIndices[i] = i;
+	}
+#endif
+
+	// QR decomposition.
+	for (int i = 0; i < FEATURE_MATRIX_SIZE; ++i)
+	{
+		NA.M[i] = DFTwoSum(A[AOffset + i], 0.0f);
+	}
+
+	QRDecomposition(NA, NQ, NR
+#if QR_DECOMPOSITION_TYPE == QR_DECOMPOSITION_TYPE_COLUMN_PIVOT
+		, PivotIndices
+#endif
+	);
+
+	// Q^T * B
+	FDFScalar Y[MATRIX_DIM * 3];
+	for (int i = 0; i < MATRIX_DIM; ++i)
+	{
+		for (int k = 0; k < 3; ++k)
+		{
+			FDFScalar Sum = (FDFScalar)0;
+			for (int j = 0; j < MATRIX_DIM; ++j)
+			{
+				Sum = DFAdd(Sum, DFMultiply(NQ.M[j * MATRIX_DIM + i], B[BOffset + j * 3 + k]));
+			}
+
+			Y[i * 3 + k] = Sum;
+		}
+	}
+
+	//back substitution for each color.
+	// RX = Y, Y = Q^T*B
+	FDFScalar X[MATRIX_DIM * 3];
+
+	for (int i = 0; i < (MATRIX_DIM * 3); ++i)
+	{
+		X[i] = DFTwoSum(0.0f, 0.0f);
+	}
+
+	for (int c = 0; c < 3; ++c)
+	{
+		for (int i = MATRIX_DIM - 1; i >= 0; --i)
+		{
+			FDFScalar Ri = NR.M[i * MATRIX_DIM + i];
+			if (DFEquals(Ri, 0.0f))
+			{
+				break;
+			}
+
+			FDFScalar Xc = Y[i * 3 + c];
+
+			for (int j = i + 1; j < MATRIX_DIM; ++j)
+			{
+				Xc = DFSubtract(Xc, DFMultiply(NR.M[i * MATRIX_DIM + j], X[j * 3 + c]));
+			}
+
+			X[i * 3 + c] = DFDivide(Xc, Ri);
+		}
+	}
+
+	//Copy to B matrix.
+#if QR_DECOMPOSITION_TYPE == QR_DECOMPOSITION_TYPE_COLUMN_PIVOT
+	for (int i = 0; i < MATRIX_DIM; ++i)
+	{
+		for (int j = 0; j < 3; ++j)
+		{
+			Result[BOffset + PivotIndices[i] * 3 + j] = DFDemote(X[i * 3 + j]);
+		}
+	}
+#else
+	for (int i = 0; i < (MATRIX_DIM * 3); ++i)
+	{
+		Result[BOffset + i] = DFDemote(X[i]);
+	}
+#endif
+
+}
+
+// Implement the standard Newton-Schulz Iterative method (float and double)
+// Initialization
+// X = I/EclideanNorm(A)
+// For n iterations
+//	Y = X ( 2*I - A * X);
+//	X = Y
+
+void NewtonSchulzIteration(int AOffset, int ASize, int BOffset)
+{
+	float X[MATRIX_DIM * MATRIX_DIM];
+	float XT[MATRIX_DIM * MATRIX_DIM];
+	float TI_AX[MATRIX_DIM * MATRIX_DIM];
+
+#define NUM_INVERSE_ITERATIONS_FLOAT 10
+
+	float EuclideanNorm = 0;
+	for (int i = 0; i < ASize; ++i)
+	{
+		EuclideanNorm += Pow2(A[AOffset + i]);
+	}
+	EuclideanNorm = sqrt(EuclideanNorm);
+
+	// Improve the initial guess with identify matrix scaled by Euclidean norm.
+	for (int i = 0; i < ADim.x; ++i)
+	{
+		for (int j = 0; j < ADim.x; ++j)
+		{
+			if (i == j)
+			{
+				X[i * ADim.x + j] = 1 / EuclideanNorm;
+			}
+			else
+			{
+				X[i * ADim.x + j] = 0;
+			}
+		}
+	}
+
+	// Matrix multiplication of float32 will overflow.
+	UNROLL
+		for (int Step = 0; Step < NUM_INVERSE_ITERATIONS_FLOAT; ++Step)
+		{
+			// 2 I - A * X
+			for (int i = 0; i < ADim.x; ++i)
+			{
+				for (int k = 0; k < ADim.x; ++k)
+				{
+					float Sum = 0;
+					for (int j = 0; j < ADim.x; ++j)
+					{
+						Sum += A[AOffset + (i * ADim.x + j)] * X[j * ADim.x + k];
+					}
+
+					TI_AX[i * ADim.x + k] = -Sum + lerp(0, 2, i == k);
+				}
+			}
+			// X * TI_AX
+			for (int i = 0; i < ADim.x; ++i)
+			{
+				for (int k = 0; k < ADim.x; ++k)
+				{
+					float Sum = 0;
+					for (int j = 0; j < ADim.x; ++j)
+					{
+						Sum += X[i * ADim.x + j] * TI_AX[j * ADim.x + k];
+					}
+					XT[i * ADim.x + k] = Sum;
+				}
+			}
+
+			// Copy back
+			for (int i = 0; i < ASize; ++i)
+			{
+				X[i] = XT[i];
+			}
+		}
+
+	//Calculate approx(A^{-1}) * B.
+	for (int i = 0; i < ADim.x; ++i)
+	{
+		for (int k = 0; k < BDim.y; ++k)
+		{
+			float Sum = 0;
+			for (int j = 0; j < ADim.y; ++j)
+			{
+				Sum += X[i * ADim.x + j] * B[BOffset + j * BDim.y + k];
+			}
+
+			Result[BOffset + i * BDim.y + k] = X[i * ADim.y + k];
+		}
+	}
+}
+
+void NewtonSchulzIterationDouble(int AOffset, int ASize, int BOffset)
+{
+	FDFScalar X[MATRIX_DIM * MATRIX_DIM];
+	FDFScalar XT[MATRIX_DIM * MATRIX_DIM];
+	FDFScalar TI_AX[MATRIX_DIM * MATRIX_DIM];
+
+#define NUM_INVERSE_ITERATIONS_DOUBLE 20
+
+	float EuclideanNorm = 0;
+	for (int i = 0; i < ASize; ++i)
+	{
+		EuclideanNorm += Pow2(A[AOffset + i]);
+	}
+	EuclideanNorm = sqrt(EuclideanNorm);
+
+	// Improve the initial guess with identify matrix scaled by Euclidean norm.
+	for (int i = 0; i < ADim.x; ++i)
+	{
+		for (int j = 0; j < ADim.x; ++j)
+		{
+			if (i == j)
+			{
+				X[i * ADim.x + j] = DFTwoSum(1 / EuclideanNorm, 0.0f);
+			}
+			else
+			{
+				X[i * ADim.x + j] = (FDFScalar)0;
+			}
+		}
+	}
+
+	for (int Step = 0; Step < NUM_INVERSE_ITERATIONS_DOUBLE; ++Step)
+	{
+		// 2 I - A * X
+		for (int i = 0; i < ADim.x; ++i)
+		{
+			for (int k = 0; k < ADim.x; ++k)
+			{
+				FDFScalar Sum = (FDFScalar)0;
+				for (int j = 0; j < ADim.x; ++j)
+				{
+					const FDFScalar DFMultiplyScalar = DFMultiply(A[AOffset + (i * ADim.x + j)], X[j * ADim.x + k]);
+					Sum = DFAdd(Sum, DFMultiplyScalar);
+				}
+
+				TI_AX[i * ADim.x + k] = DFSubtract(lerp(0, 2, i == k), Sum);
+			}
+		}
+		// X * TI_AX
+		for (int i = 0; i < ADim.x; ++i)
+		{
+			for (int k = 0; k < ADim.x; ++k)
+			{
+				FDFScalar Sum = (FDFScalar)0;
+				for (int j = 0; j < ADim.x; ++j)
+				{
+					const FDFScalar DFMultiplyScalar = DFMultiply(X[i * ADim.x + j], TI_AX[j * ADim.x + k]);
+					Sum = DFAdd(Sum, DFMultiplyScalar);
+				}
+				XT[i * ADim.x + k] = Sum;
+			}
+		}
+
+		// copy back
+		for (int i = 0; i < ASize; ++i)
+		{
+			X[i] = XT[i];
+		}
+	}
+
+	//Calculate approx(A^{-1}) * B.
+	for (int i = 0; i < ADim.x; ++i)
+	{
+		for (int k = 0; k < BDim.y; ++k)
+		{
+			FDFScalar Sum = DFTwoSum(0, 0);
+			for (int j = 0; j < ADim.y; ++j)
+			{
+				Sum = DFAdd(Sum, DFMultiply(X[i * ADim.x + j], DFTwoSum(B[BOffset + j * BDim.y + k], 0.0f)));
+			}
+
+			Result[BOffset + i * BDim.y + k] = DFDemote(Sum);
+		}
+	}
+}
\ No newline at end of file
