Fix random salt and pepper specular noise for the NFOR denoiser.

Even when radiance is as expected, albedo and normal could become very large due to rare paths. Adjust their value range and scale the variance accordingly based on r.NFOR.Feature.MaxAlbedoGreyscale (default: 2.0) and r.NFOR.Feature.MaxNormalLength (default: 10.0).

diff --git a/Engine/Plugins/Experimental/NFORDenoise/Source/NFORDenoise/Private/NFORDenoiseCS.cpp b/Engine/Plugins/Experimental/NFORDenoise/Source/NFORDenoise/Private/NFORDenoiseCS.cpp
--- a/Engine/Plugins/Experimental/NFORDenoise/Source/NFORDenoise/Private/NFORDenoiseCS.cpp
+++ b/Engine/Plugins/Experimental/NFORDenoise/Source/NFORDenoise/Private/NFORDenoiseCS.cpp
@@ -16,6 +16,20 @@ namespace NFORDenoise
 		TEXT("Add a constant 1 feature for denoising. Especially useful when all other features are zero."),
 		ECVF_RenderThreadSafe);
 
+	TAutoConsoleVariable<float> CVarNFORFeatureMaxAlbedoGreyscale(
+		TEXT("r.NFOR.Feature.MaxAlbedoGreyscale"),
+		2.0,
+		TEXT("Set the max albedo in greyscale used for denoising. Scale the albedo variance as well. Used for suppressing specular noise.")
+		TEXT("<=0: Ignore scaling."),
+		ECVF_RenderThreadSafe);
+
+	TAutoConsoleVariable<float> CVarNFORFeatureMaxNormalLength(
+		TEXT("r.NFOR.Feature.MaxNormalLength"),
+		10.0,
+		TEXT("Set the max normal length used for denoising. Scale the normal variance as well. Used for suppressing specular noise.")
+		TEXT("<=0: Ignore scaling."),
+		ECVF_RenderThreadSafe);
+
 	TAutoConsoleVariable<int32> CVarNFORPredivideAlbedo(
 		TEXT("r.NFOR.PredivideAlbedo"),
 		1,
@@ -165,6 +179,16 @@ namespace NFORDenoise
 		return CVarNFORFeatureAddConstant.GetValueOnRenderThread();
 	}
 
+	float GetFeatureMaxAlbedoGrayscale()
+	{
+		return CVarNFORFeatureMaxAlbedoGreyscale.GetValueOnRenderThread();
+	}
+
+	float GetFeatureMaxNormalLength()
+	{
+		return CVarNFORFeatureMaxNormalLength.GetValueOnRenderThread();
+	}
+
 	bool IsPreAlbedoDivideEnabled()
 	{
 		const bool bPredivideAlbedo = CVarNFORPredivideAlbedo.GetValueOnRenderThread() != 0;
@@ -378,9 +402,10 @@ namespace NFORDenoise
 	IMPLEMENT_GLOBAL_SHADER(FCopyTexturePS, "/NFORDenoise/NFORDenoise.usf", "CopyTexturePS", SF_Pixel);
 
 	//--------------------------------------------------------------------------------------------------------------------
-	// Radiance normalization
+	// Feature range adjustment and radiance normalization
 	IMPLEMENT_GLOBAL_SHADER(FClassifyPreAlbedoDivideMaskIdCS, "/NFORDenoise/NFORDenoise.usf", "ClassifyPreAlbedoDivideMaskIdCS", SF_Compute);
 	IMPLEMENT_GLOBAL_SHADER(FNormalizeRadianceVarianceByAlbedoCS, "/NFORDenoise/NFORDenoise.usf", "NormalizeRadianceVarianceByAlbedoCS", SF_Compute);
+	IMPLEMENT_GLOBAL_SHADER(FAdjustFeatureRangeCS, "/NFORDenoise/NFORDenoise.usf", "AdjustFeatureRangeCS", SF_Compute);
 
 	//--------------------------------------------------------------------------------------------------------------------
 	// Non-local mean weight and filtering
@@ -655,6 +680,47 @@ namespace NFORDenoise
 		return MaskTexture;
 	}
 
+	void AddAdjustFeatureRangePass(FRDGBuilder& GraphBuilder, const FFeatureDesc& FeatureDesc, float MaxValue)
+	{
+		checkf(FeatureDesc.Data.NumOfChannel == 4, TEXT("Only feature with 4 channels can be adjusted"));
+		checkf(FeatureDesc.VarianceType != EVarianceType::Colored, TEXT("Feature variance of type EVarianceType::Colored cannot be adjusted"));
+		
+		if (MaxValue <= 0)
+		{
+			return;
+		}
+
+		FIntPoint Size = FeatureDesc.Data.Image->Desc.Extent;
+		typedef FAdjustFeatureRangeCS SHADER;
+		SHADER::FParameters* PassParameters = GraphBuilder.AllocParameters<SHADER::FParameters>();
+		{
+			PassParameters->RWImage = GraphBuilder.CreateUAV(FeatureDesc.Data.Image);
+			PassParameters->RWImageVariance = GraphBuilder.CreateUAV(FeatureDesc.Variance.Image);
+			PassParameters->Size = Size;
+			PassParameters->VarianceChannelOffset = FeatureDesc.Variance.ChannelOffset;
+			PassParameters->MaxValue = MaxValue;
+		}
+
+		SHADER::FPermutationDomain ComputeShaderPermutationVector;
+		{
+			ComputeShaderPermutationVector.Set<SHADER::FDimensionVarianceType>(FeatureDesc.VarianceType);
+		}
+
+		TShaderMapRef<SHADER> ComputeShader(GetGlobalShaderMap(GMaxRHIFeatureLevel), ComputeShaderPermutationVector);
+
+		FComputeShaderUtils::AddPass(
+			GraphBuilder,
+			RDG_EVENT_NAME("NFOR::AddAdjustFeatureRangePass (%s, MaxValue=%.2f, size:%dx%d)",
+				FeatureDesc.Data.Image->Name,
+				MaxValue,
+				Size.X,
+				Size.Y),
+			ERDGPassFlags::Compute,
+			ComputeShader,
+			PassParameters,
+			FComputeShaderUtils::GetGroupCount(Size, NON_LOCAL_MEAN_THREAD_GROUP_SIZE));
+	}
+
 	//--------------------------------------------------------------------------------------------------------------------
 	// Non-local mean weight and filtering
 	FNonLocalMeanParameters GetNonLocalMeanParameters(int32 PatchSize, int32 PatchDistance, float Bandwidth)
@@ -1591,16 +1657,23 @@ namespace NFORDenoise
 		const int32 SourceRadianceIndex = GetDenoisingFrameIndex(NumOfTemporalFrames);
 		
 		// Preprocessing
-		// Radiance normalization and filtering frames
+		// Feature range adjustment, radiance normalization and filtering frames
 		{
-			if (IsPreAlbedoDivideEnabled())
+			// Latest frame only 
+			for (int i = 0; i < 1; ++i)
 			{
-				// Latest frame only 
-				for (int i = 0; i < 1; ++i)
+				const FFeatureDesc& Albedo = FeatureDescs[i * NumOfFeaturesPerFrame + 0]; //TODO: unify the index
+				const FFeatureDesc& Normal = FeatureDescs[i * NumOfFeaturesPerFrame + 1]; //TODO: unify the index
+
+				// Adjust feature range if required
+				AddAdjustFeatureRangePass(GraphBuilder, Albedo, GetFeatureMaxAlbedoGrayscale());
+				AddAdjustFeatureRangePass(GraphBuilder, Normal, GetFeatureMaxNormalLength());
+
+				if (IsPreAlbedoDivideEnabled())
 				{
-					FRDGTextureRef AlbedoTex = FeatureDescs[i * NumOfFeaturesPerFrame + 0].Feature.Image; //TODO: unify the index
-					FRDGTextureRef NormalTex = FeatureDescs[i * NumOfFeaturesPerFrame + 1].Feature.Image; //TODO: unify the index
-					FRDGTextureRef NormalVarianceTex = FeatureDescs[i * NumOfFeaturesPerFrame + 1].Variance.Image;
+					FRDGTextureRef AlbedoTex = Albedo.Feature.Image; 
+					FRDGTextureRef NormalTex = Normal.Feature.Image; 
+					FRDGTextureRef NormalVarianceTex = Normal.Variance.Image;
 					FRDGTextureRef MaskTexture = NFORDenoise::GetPreAlbedoDivideMask(GraphBuilder, View, NormalTex, NormalVarianceTex);
 
 					FLinearColor RGBOffset = NFORDenoise::GetPreAlbedoDivideAlbedoOffset();
@@ -1610,7 +1683,7 @@ namespace NFORDenoise
 					FRDGTextureRef RadianceTexture = Radiances[i].Data.Image;
 					FRDGTextureRef RadianceVarianceTexture = Radiances[i].Variance.Image;
 
-					// Normalization should apply to both texture to variance.
+					// Normalization should apply to both texture and variance.
 					NFORDenoise::AddDivideTextureRegionPass(GraphBuilder, AlbedoTex, RadianceTexture);
 					NFORDenoise::AddNormalizeRadianceVariancePass(GraphBuilder, AlbedoTex, RadianceVarianceTexture);
 				}
@@ -1745,4 +1818,4 @@ namespace NFORDenoise
 
 		return true;
 	}
-}
\ No newline at end of file
+}
diff --git a/Engine/Plugins/Experimental/NFORDenoise/Source/NFORDenoise/Private/NFORDenoiseCS.h b/Engine/Plugins/Experimental/NFORDenoise/Source/NFORDenoise/Private/NFORDenoiseCS.h
--- a/Engine/Plugins/Experimental/NFORDenoise/Source/NFORDenoise/Private/NFORDenoiseCS.h
+++ b/Engine/Plugins/Experimental/NFORDenoise/Source/NFORDenoise/Private/NFORDenoiseCS.h
@@ -179,7 +179,7 @@ namespace NFORDenoise
 	
 
 	//--------------------------------------------------------------------------------------------------------------------
-	// Radiance normalization by albedo
+	// Feature range adjustment, and radiance normalization by albedo
 	
 	class FClassifyPreAlbedoDivideMaskIdCS : public FGlobalShader
 	{
@@ -221,10 +221,41 @@ namespace NFORDenoise
 		}
 	};
 
+	struct FFeatureDesc;
+
+	// Remap feature range
+	class FAdjustFeatureRangeCS : public FGlobalShader
+	{
+		DECLARE_GLOBAL_SHADER(FAdjustFeatureRangeCS)
+		SHADER_USE_PARAMETER_STRUCT(FAdjustFeatureRangeCS, FGlobalShader)
+	public:
+
+		BEGIN_SHADER_PARAMETER_STRUCT(FParameters, )
+			SHADER_PARAMETER_RDG_TEXTURE_UAV(RWTexture2D, RWImage)
+			SHADER_PARAMETER_RDG_TEXTURE_UAV(RWTexture2D, RWImageVariance)
+			SHADER_PARAMETER(FIntPoint, Size)
+			SHADER_PARAMETER(int32, VarianceChannelOffset)
+			SHADER_PARAMETER(float, MaxValue)
+		END_SHADER_PARAMETER_STRUCT()
+
+		static void ModifyCompilationEnvironment(const FGlobalShaderPermutationParameters& InParameters, FShaderCompilerEnvironment& OutEnvironment)
+		{
+			FGlobalShader::ModifyCompilationEnvironment(InParameters, OutEnvironment);
+			OutEnvironment.SetDefine(TEXT("THREAD_GROUP_SIZE"), NON_LOCAL_MEAN_THREAD_GROUP_SIZE);
+			OutEnvironment.CompilerFlags.Add(CFLAG_AllowTypedUAVLoads);
+		}
+
+		class FDimensionVarianceType : SHADER_PERMUTATION_ENUM_CLASS("IMAGE_VARIANCE_TYPE", EVarianceType);
+		using FPermutationDomain = TShaderPermutationDomain<FDimensionVarianceType>;
+	};
+
 	FRDGTextureRef GetPreAlbedoDivideMask(FRDGBuilder& GraphBuilder, const FSceneView& View, const FRDGTextureRef& Normal, const FRDGTextureRef& NormalVariance);
 
 	void AddNormalizeRadianceVariancePass(FRDGBuilder& GraphBuilder, const FRDGTextureRef& Albedo, const FRDGTextureRef& RadianceVariance);
 
+	// Adjust feature and variance range based on MaxValue to suppress random specular noise.
+	// NewValue = min(Value, MaxValue). Std = lerp(Std,(MaxValue/Value)*Std,Value > MaxValue).
+	void AddAdjustFeatureRangePass(FRDGBuilder& GraphBuilder, const FFeatureDesc& FeatureDesc, float MaxValue);
 
 	//--------------------------------------------------------------------------------------------------------------------
 	// Non-local mean weight and filtering
diff --git a/Engine/Plugins/Experimental/NFORDenoise/Source/NFORDenoise/Shaders/NFORDenoise.usf b/Engine/Plugins/Experimental/NFORDenoise/Source/NFORDenoise/Shaders/NFORDenoise.usf
--- a/Engine/Plugins/Experimental/NFORDenoise/Source/NFORDenoise/Shaders/NFORDenoise.usf
+++ b/Engine/Plugins/Experimental/NFORDenoise/Source/NFORDenoise/Shaders/NFORDenoise.usf
@@ -268,6 +268,40 @@ void NormalizeRadianceVarianceByAlbedoCS(in const uint3 DispatchThreadID : SV_Di
 	RWRadianceVariance[Position].r /= lerp(1.0f, AlbedoLum, AlbedoLum > 0);
 }
 
+RWTexture2D<float4> RWImage;
+RWTexture2D<float4> RWImageVariance;
+float MaxValue;
+int VarianceChannelOffset;
+
+[numthreads(THREAD_GROUP_SIZE, THREAD_GROUP_SIZE, 1)]
+void AdjustFeatureRangeCS(in const uint3 DispatchThreadID : SV_DispatchThreadID)
+{
+	if (any(DispatchThreadID.xy >= (uint2)Size))
+	{
+		return;
+	}
+
+	const uint2 Position = DispatchThreadID.xy;
+
+	float3 Feature = RWImage[Position].xyz;
+#if (IMAGE_VARIANCE_TYPE == IMAGE_VARIANCE_NORMAL)
+	float Value = length(Feature);
+#elif (IMAGE_VARIANCE_TYPE == IMAGE_VARIANCE_GREYSCALE)
+	float Value = Luminance(Feature);
+#else
+	float Value = -1.0f;
+#endif
+
+	const float ScaleFactor = lerp(1.0f, MaxValue / Value, Value > MaxValue);
+	
+	BRANCH
+	if (ScaleFactor < 1.0f)
+	{
+		RWImage[Position].xyz *= ScaleFactor;
+		RWImageVariance[Position][VarianceChannelOffset] *= ScaleFactor; // Scaling the std does not require power.
+	}
+}
+
 //--------------------------------------------------------------------------------------------------
 // Non-local mean definition and function
 //--------------------------------------------------------------------------------------------------
@@ -289,8 +323,6 @@ Texture2D<TImageVariance> TargetVariance;
 #define TargetVariance Variance
 #endif
 
-int VarianceChannelOffset;
-
 float GetSquaredDistance(int2 P, int2 Q, float VarianceScale = 1.0f)
 {
 	// Use mirrored position. Clamp to boarder has artifacts at border.
