Improve the non-local mean weight query for the NFOR denoiser with atlas and separable filter.
* Atlas is used to cache patch differencce to increase parallelism. Separable filter is then applied to perform fast box filter to calculate the patch weights.
* Atlas is turned on by default (r.NFOR.NonLocalMean.Atlas true), and the atlas size is default to 2k.

#tests Hillside static shots rendering stays the same

diff --git a/Engine/Plugins/Experimental/NFORDenoise/Source/NFORDenoise/Private/NFORDenoise.cpp b/Engine/Plugins/Experimental/NFORDenoise/Source/NFORDenoise/Private/NFORDenoise.cpp
--- a/Engine/Plugins/Experimental/NFORDenoise/Source/NFORDenoise/Private/NFORDenoise.cpp
+++ b/Engine/Plugins/Experimental/NFORDenoise/Source/NFORDenoise/Private/NFORDenoise.cpp
@@ -124,6 +124,8 @@ public:
 				return;
 			}
 
+			RDG_EVENT_SCOPE(GraphBuilder, "AddFrame");
+
 			FRDGFrameCache RDGFrameCache;
 			RDGFrameCache.Radiance = GraphBuilder.CreateTexture(Inputs.ColorTex->Desc,TEXT("Radiance"));
 			RDGFrameCache.RadianceVariance = GraphBuilder.CreateTexture(Inputs.VarianceTex->Desc, TEXT("RadianceVariance"));
diff --git a/Engine/Plugins/Experimental/NFORDenoise/Source/NFORDenoise/Private/NFORDenoiseCS.cpp b/Engine/Plugins/Experimental/NFORDenoise/Source/NFORDenoise/Private/NFORDenoiseCS.cpp
--- a/Engine/Plugins/Experimental/NFORDenoise/Source/NFORDenoise/Private/NFORDenoiseCS.cpp
+++ b/Engine/Plugins/Experimental/NFORDenoise/Source/NFORDenoise/Private/NFORDenoiseCS.cpp
@@ -30,6 +30,13 @@ namespace NFORDenoise
 		TEXT("<=0: Ignore scaling."),
 		ECVF_RenderThreadSafe);
 
+	TAutoConsoleVariable<bool> CVarNFORFeatureFiltering(
+		TEXT("r.NFOR.Feature.Filtering"),
+		true,
+		TEXT("True: Filter all features.\n")
+		TEXT("False: Disable feature filtering (useful for debug).\n"),
+		ECVF_RenderThreadSafe);
+
 	TAutoConsoleVariable<int32> CVarNFORPredivideAlbedo(
 		TEXT("r.NFOR.PredivideAlbedo"),
 		1,
@@ -123,6 +130,27 @@ namespace NFORDenoise
 		TEXT(">=0: Output the denoising contribution from the ith frame only.")
 		TEXT("-1: do not perform debug. Output contributions from all frames."),
 		ECVF_RenderThreadSafe);
+	
+	TAutoConsoleVariable<bool> CVarNFORNonLocalMeanAtlas(
+		TEXT("r.NFOR.NonLocalMean.Atlas"),
+		true,
+		TEXT("true	: Use atlas and separable filter to improve the performance of NLM weights query.\n")
+		TEXT("false : Calculate the non local mean weights for each pixel in place.(baseline but slow).\n"),
+		ECVF_RenderThreadSafe);
+
+	TAutoConsoleVariable<int32> CVarNFORNonLocalMeanAtlasType(
+		TEXT("r.NFOR.NonLocalMean.Atlas.Type"),
+		1,
+		TEXT("0: float2. Stores one symmetric distance/weight.\n")
+		TEXT("1: float4. Stores two symmetric distance/weights.\n"),
+		ECVF_RenderThreadSafe);
+
+	TAutoConsoleVariable<int32> CVarNFORNonLocalMeanAtlasSize(
+		TEXT("r.NFOR.NonLocalMean.Atlas.Size"),
+		2048,
+		TEXT("<=0: Use the same size of the input tile.\n")
+		TEXT("	n: At least the max size of the input tile(2k as default). The larger, the less number of dispatch passes.\n"),
+		ECVF_RenderThreadSafe);
 
 	TAutoConsoleVariable<int32> CVarNFORNonLocalMeanFeaturePatchSize(
 		TEXT("r.NFOR.NonLocalMean.Feature.PatchSize"),
@@ -189,6 +217,11 @@ namespace NFORDenoise
 		return CVarNFORFeatureMaxNormalLength.GetValueOnRenderThread();
 	}
 
+	bool ShouldApplyFeatureFiltering()
+	{
+		return CVarNFORFeatureFiltering.GetValueOnRenderThread();
+	}
+
 	bool IsPreAlbedoDivideEnabled()
 	{
 		const bool bPredivideAlbedo = CVarNFORPredivideAlbedo.GetValueOnRenderThread() != 0;
@@ -355,6 +388,35 @@ namespace NFORDenoise
 		return FMath::Clamp(CVarNFORNonLocalMeanRadiancePatchDistance.GetValueOnRenderThread(), 0, 30);
 	}
 
+	FIntPoint GetNonLocalMeanAtlasSize(FIntPoint Extent)
+	{
+		int32 AtlasSize = CVarNFORNonLocalMeanAtlasSize.GetValueOnRenderThread();
+		if (AtlasSize <= 0)
+		{
+			return Extent;
+		}
+		else
+		{
+			AtlasSize = FMath::Max(Extent.GetMax(),AtlasSize);
+			return FIntPoint(AtlasSize,AtlasSize);
+		}
+	}
+
+	bool ShouldNonLocalMeanUseAtlas()
+	{
+		return CVarNFORNonLocalMeanAtlas.GetValueOnRenderThread();
+	}
+
+	ENonLocalMeanAtlasType GetNonLocalMeanAtlasType()
+	{
+		const uint32 AtlasType = FMath::Clamp(
+			CVarNFORNonLocalMeanAtlasType.GetValueOnRenderThread(),
+			0,
+			static_cast<uint32>(ENonLocalMeanAtlasType::MAX)-1
+			);
+		return static_cast<ENonLocalMeanAtlasType>(AtlasType);
+	}
+
 	bool IsBandwidthSelectionEnabled()
 	{
 		return CVarNFORBandwidthSelection.GetValueOnRenderThread();
@@ -412,6 +474,11 @@ namespace NFORDenoise
 	IMPLEMENT_GLOBAL_SHADER(FNonLocalMeanFilteringCS, "/NFORDenoise/NFORDenoise.usf", "NonLocalMeanFilteringCS", SF_Compute);
 	IMPLEMENT_GLOBAL_SHADER(FNonLocalMeanWeightsCS, "/NFORDenoise/NFORDenoise.usf", "NonLocalMeanWeightsCS", SF_Compute);
 
+	// Fast weights query.
+	IMPLEMENT_GLOBAL_SHADER(FNonLocalMeanGetSqauredDistanceToAtlasCS, "/NFORDenoise/NFORDenoise.usf", "NonLocalMeanGetSqauredDistanceToAtlasCS", SF_Compute);
+	IMPLEMENT_GLOBAL_SHADER(FNonLocalMeanSeperableFilterPatchSqauredDistanceCS, "/NFORDenoise/NFORDenoise.usf", "NonLocalMeanSeperableFilterPatchSqauredDistanceCS", SF_Compute);
+	IMPLEMENT_GLOBAL_SHADER(FNonLocalMeanReshapeBufferCS, "/NFORDenoise/NFORDenoise.usf", "NonLocalMeanReshapeBufferCS", SF_Compute);
+
 	//--------------------------------------------------------------------------------------------------------------------
 	// Collaborative filtering
 	//	1. Tiling
@@ -746,7 +813,7 @@ namespace NFORDenoise
 		const FNonLocalMeanParameters& NonLocalMeanParameters,
 		const FNFORTextureDesc& Texture,
 		const FNFORTextureDesc& Variance,
-		EVarianceType VarianceTyle,
+		EVarianceType VarianceType,
 		const FRDGTextureRef& FilteredTexture)
 	{
 		FGlobalShaderMap* GlobalShaderMap = GetGlobalShaderMap(GMaxRHIFeatureLevel);
@@ -768,7 +835,7 @@ namespace NFORDenoise
 
 		SHADER::FPermutationDomain ComputeShaderPermutationVector;
 		{
-			ComputeShaderPermutationVector.Set<SHADER::FDimensionVarianceType>(VarianceTyle);
+			ComputeShaderPermutationVector.Set<SHADER::FDimensionVarianceType>(VarianceType);
 			ComputeShaderPermutationVector.Set<SHADER::FDimensionUseGuide>(false);
 			ComputeShaderPermutationVector.Set<SHADER::FDimensionImageChannelCount>(Texture.NumOfChannel);
 			ComputeShaderPermutationVector.Set<SHADER::FDimPreAlbedoDivide>(GetPreAlbedoDivideRecoverPhase());
@@ -785,7 +852,31 @@ namespace NFORDenoise
 			FComputeShaderUtils::GetGroupCount(TextureSize, NON_LOCAL_MEAN_THREAD_GROUP_SIZE));
 	}
 
-	void GetNLMWeights(
+	void ApplyNonLocalMeanFilterIfRequired(
+		FRDGBuilder& GraphBuilder,
+		const FNonLocalMeanParameters& NonLocalMeanParameters,
+		const FNFORTextureDesc& Texture,
+		const FNFORTextureDesc& Variance,
+		EVarianceType VarianceType,
+		const FRDGTextureRef& FilteredTexture)
+	{
+		if (ShouldApplyFeatureFiltering())
+		{
+			ApplyNonLocalMeanFilter(
+				GraphBuilder,
+				NonLocalMeanParameters,
+				Texture,
+				Variance,
+				VarianceType,
+				FilteredTexture);
+		}
+		else
+		{
+			AddCopyTexturePass(GraphBuilder, Texture.Image, FilteredTexture);
+		}
+	}
+
+	void GetNLMWeigthsWithAtlas(
 		FRDGBuilder& GraphBuilder,
 		const FSceneView& View,
 		const FRadianceDesc& SourceRadiance,
@@ -798,8 +889,250 @@ namespace NFORDenoise
 		const int32 NumberOfWeightsPerPixel = SearchingPatchSize * SearchingPatchSize;
 		const FIntPoint TextureSize = SourceRadiance.Data.Image->Desc.Extent;
 		const bool bSeparateSourceTarget = (SourceRadiance.Data.Image != TargetRadiance.Data.Image);
+
+		const FIntRect SeparableFilteringRegion = Region.Inner(-(NonLocalMeanParameters.PatchDistance + NonLocalMeanParameters.PatchSize));
+		const FIntPoint SeparableFilteringExtent = SeparableFilteringRegion.Size();
+		const FIntPoint WeightQueryRegionExtent = Region.Inner(-(NonLocalMeanParameters.PatchDistance)).Size();
+		
+		// Estimate the number of dispatches required using the atlas
+		const FIntPoint NLMWeightAtlasExtent = GetNonLocalMeanAtlasSize(SeparableFilteringExtent);
+		const FIntVector DispatchTileVector = FIntVector(
+			FMath::DivideAndRoundDown(NLMWeightAtlasExtent.X, SeparableFilteringExtent.X),
+			FMath::DivideAndRoundDown(NLMWeightAtlasExtent.Y, SeparableFilteringExtent.Y), 1);
+
+		const int32 HalfOffsetSearchCount = (NumberOfWeightsPerPixel / 2 + 1);
+		const int32 DispatchTileCount = DispatchTileVector.X * DispatchTileVector.Y;
+
+		const ENonLocalMeanAtlasType NonLocalMeanAtlasType = GetNonLocalMeanAtlasType();
+		const int32 NumSymmetricPairsPerPixel = static_cast<int32>(NonLocalMeanAtlasType) + 1;
+		const int32 SingleDispatchOffsetSearchCount = DispatchTileCount * NumSymmetricPairsPerPixel;
+		const int32 NumOfDispatch = FMath::DivideAndRoundUp(HalfOffsetSearchCount, SingleDispatchOffsetSearchCount);
+
+		// Allocate the atlas where each pixel stores two symmetric distance/weights, and the temporal buffer.
+		FRDGTextureDesc NLMWeightAtlasDesc= SourceRadiance.Data.Image->Desc;
+		int32 NonLocalMeanWeightsBytesPerElement = 0;
+		{
+			NLMWeightAtlasDesc.Extent = NLMWeightAtlasExtent;
+			switch (NonLocalMeanAtlasType)
+			{
+			case ENonLocalMeanAtlasType::OneSymmetricPair:
+				NLMWeightAtlasDesc.Format = PF_G32R32F;
+				NonLocalMeanWeightsBytesPerElement = 2 * sizeof(float);
+				break;
+			case ENonLocalMeanAtlasType::TwoSymmetricPair:
+			default:
+				NonLocalMeanWeightsBytesPerElement = 4 * sizeof(float);
+				NLMWeightAtlasDesc.Format = PF_A32B32G32R32F; 
+				break;
+			}
+		}
 		
+		FRDGTextureRef NLMWeightAtlas[2] = {
+			GraphBuilder.CreateTexture(NLMWeightAtlasDesc, TEXT("NFOR.NLMWeightAtlas0")),
+			GraphBuilder.CreateTexture(NLMWeightAtlasDesc, TEXT("NFOR.NLMWeightAtlas1"))
+		};
+		
+		const int32 TotalNumTilesToFill = FMath::DivideAndRoundUp(HalfOffsetSearchCount, NumSymmetricPairsPerPixel);
+		size_t NonLocalMeanWeightsCount = NumSymmetricPairsPerPixel * (WeightQueryRegionExtent.X * WeightQueryRegionExtent.Y) * TotalNumTilesToFill;
+		FRDGBufferDesc NonLocalMeanWeightsDesc = FRDGBufferDesc::CreateBufferDesc(NonLocalMeanWeightsBytesPerElement, NonLocalMeanWeightsCount);
+		FRDGBufferRef  NonLocalMeanWeights = GraphBuilder.CreateBuffer(NonLocalMeanWeightsDesc, TEXT("NFOR.NonLocalMeanWeights"));
+
+		// Summery
+		// 1. For each dispatch:
+		//		Calcualte offset for each tile (SeparableFilteringRegion) in the atlas
+		//		Horizontal filter to second atlas
+		//		Vertical filter to buffer
+		// 2. Reshape the buffer for later use
+
+
+		struct FSeperableFilterPassInfo
+		{
+			const TCHAR* PassName;
+			FRDGTextureRef Input;
+			FRDGTextureRef Output;
+			FIntPoint	GroupCountXY;
+			FNonLocalMeanSeperableFilterPatchSqauredDistanceCS::ESeperablePassType SeperablePassType;
+		};
+
+		const int NumOfSeperablePass = 2;
+
+		//	Horizontal requires all rows
+		//	Vertical filtering only requires the weight query regions and stores to a buffer.
+		const FSeperableFilterPassInfo SeperableFilterPassInfo[NumOfSeperablePass] =
+		{
+			{	TEXT("NFOR::SeperableHorizontal"),	NLMWeightAtlas[0], NLMWeightAtlas[1], FIntPoint(SeparableFilteringExtent.X, SeparableFilteringExtent.Y),
+			FNonLocalMeanSeperableFilterPatchSqauredDistanceCS::ESeperablePassType::Horizontal},			//Separable horizontal
+			{	TEXT("NFOR::SeperableVertical"),	NLMWeightAtlas[1],			 nullptr, FIntPoint(WeightQueryRegionExtent.X, WeightQueryRegionExtent.Y),
+			FNonLocalMeanSeperableFilterPatchSqauredDistanceCS::ESeperablePassType::Vertical},				//Separable Vertical
+		};
+
+		for (int32 DispatchId = 0; DispatchId < NumOfDispatch; ++DispatchId)
+		{
+			const int32 NumSymmetricPairsPerDispatch = SingleDispatchOffsetSearchCount - FMath::Max((DispatchId + 1) * SingleDispatchOffsetSearchCount - HalfOffsetSearchCount, 0);
+			const int32 NumOfTilesToFillThisDispatch = FMath::DivideAndRoundUp(NumSymmetricPairsPerDispatch, NumSymmetricPairsPerPixel);
+			const FIntVector DispatchRegionSize = FIntVector(SeparableFilteringExtent.X, SeparableFilteringExtent.Y, NumOfTilesToFillThisDispatch);
+
+			FNonLocalMeanWeightAtlasDispatchParameters NLMWeightAtlasDispatchParameters;
+			{
+				NLMWeightAtlasDispatchParameters.DispatchId = DispatchId;
+				NLMWeightAtlasDispatchParameters.DispatchTileSize = FIntPoint(DispatchTileVector.X, DispatchTileVector.Y);
+				NLMWeightAtlasDispatchParameters.DispatchTileCount = DispatchTileCount;
+				NLMWeightAtlasDispatchParameters.SeparableFilteringRegion = SeparableFilteringRegion;
+				NLMWeightAtlasDispatchParameters.DispatchRegionSize = DispatchRegionSize;
+			}
+
+			// Get squared distance, each tile pixel holds NumSymmetricPairsPerPixel symmetric pairs.
+			{
+				typedef FNonLocalMeanGetSqauredDistanceToAtlasCS SHADER;
+
+				SHADER::FParameters* PassParameters = GraphBuilder.AllocParameters<SHADER::FParameters>();
+				{
+					PassParameters->CommonParameters.NLMParams = NonLocalMeanParameters;
+					PassParameters->NLMWeightAtlasDispatchParameters = NLMWeightAtlasDispatchParameters;
+					PassParameters->CommonParameters.Image = GraphBuilder.CreateSRV(FRDGTextureSRVDesc::Create(SourceRadiance.Data.Image));
+					PassParameters->CommonParameters.Variance = GraphBuilder.CreateSRV(FRDGTextureSRVDesc::Create(SourceRadiance.Variance.Image));
+					PassParameters->CommonParameters.TextureSize = TextureSize;
+					PassParameters->CommonParameters.VarianceChannelOffset = SourceRadiance.Variance.ChannelOffset;
+
+					if (bSeparateSourceTarget)
+					{
+						PassParameters->TargetImage = GraphBuilder.CreateSRV(FRDGTextureSRVDesc::Create(TargetRadiance.Data.Image));
+						PassParameters->TargetVariance = GraphBuilder.CreateSRV(FRDGTextureSRVDesc::Create(TargetRadiance.Variance.Image));
+					}
+					PassParameters->RWNLMWeightAtlas = GraphBuilder.CreateUAV(NLMWeightAtlas[0]);
+					PassParameters->NLMWeightAtlasSize = NLMWeightAtlasExtent;
+				}
+
+				SHADER::FPermutationDomain ComputeShaderPermutationVector;
+				{
+					ComputeShaderPermutationVector.Set<SHADER::FDimensionVarianceType>(SourceRadiance.VarianceType);
+					ComputeShaderPermutationVector.Set<SHADER::FDimensionImageChannelCount>(SourceRadiance.Data.NumOfChannel);
+					ComputeShaderPermutationVector.Set<SHADER::FDimensionSeparateSourceTarget>(bSeparateSourceTarget);
+					ComputeShaderPermutationVector.Set<SHADER::FDimAtlasType>(NonLocalMeanAtlasType);
+					ComputeShaderPermutationVector.Set<SHADER::FDimPreAlbedoDivide>(GetPreAlbedoDivideRecoverPhase());
+				}
+
+				TShaderMapRef<SHADER> ComputeShader(GetGlobalShaderMap(GMaxRHIFeatureLevel), ComputeShaderPermutationVector);
+				
+				FComputeShaderUtils::AddPass(
+					GraphBuilder,
+					RDG_EVENT_NAME("NFOR::NonLocalMeanGetSqauredDistanceToAtlasCS (Rect=(%d,%d,%d,%d), pd=%d, DisaptchId,DispatchTileCount=%d,%d, GC=(%d,%d,%d))",
+						SeparableFilteringRegion.Min.X,
+						SeparableFilteringRegion.Min.Y,
+						SeparableFilteringRegion.Max.X,
+						SeparableFilteringRegion.Max.Y,
+						NonLocalMeanParameters.PatchDistance,
+						DispatchId,
+						DispatchTileCount,
+						DispatchRegionSize.X,
+						DispatchRegionSize.Y,
+						DispatchRegionSize.Z),
+					ERDGPassFlags::Compute,
+					ComputeShader,
+					PassParameters,
+					FComputeShaderUtils::GetGroupCount(DispatchRegionSize, NON_LOCAL_MEAN_THREAD_GROUP_SIZE));
+			}
+
+			// Horizontal and vertical filtering
+			for (int SeperableFilteringId = 0; SeperableFilteringId < NumOfSeperablePass; ++SeperableFilteringId)
+			{
+				const FSeperableFilterPassInfo& PassInfo = SeperableFilterPassInfo[SeperableFilteringId];
+				FIntVector SeperableRegionSize = FIntVector(PassInfo.GroupCountXY.X, PassInfo.GroupCountXY.Y, NumOfTilesToFillThisDispatch);
+
+				typedef FNonLocalMeanSeperableFilterPatchSqauredDistanceCS SHADER;
+
+				SHADER::FParameters* PassParameters = GraphBuilder.AllocParameters<SHADER::FParameters>();
+				{
+					PassParameters->NLMParams = NonLocalMeanParameters;
+					PassParameters->NLMWeightAtlasDispatchParameters = NLMWeightAtlasDispatchParameters;
+					PassParameters->NLMWeightAtlasSource = GraphBuilder.CreateSRV(PassInfo.Input);
+
+					if (PassInfo.SeperablePassType == SHADER::ESeperablePassType::Horizontal)
+					{
+						PassParameters->RWNLMWeightAtlasTarget = GraphBuilder.CreateUAV(PassInfo.Output);
+					}
+					else
+					{
+						// Vertical pass directly write to the weight buffer
+						PassParameters->RWNLMWeights = GraphBuilder.CreateUAV(FRDGBufferUAVDesc(NonLocalMeanWeights, NLMWeightAtlasDesc.Format));
+					}
+					PassParameters->SeperableRegionSize = SeperableRegionSize;
+				}
+
+				SHADER::FPermutationDomain ComputeShaderPermutationVector;
+				{
+					ComputeShaderPermutationVector.Set<SHADER::FDimensionSeperablePassType>(PassInfo.SeperablePassType);
+					ComputeShaderPermutationVector.Set<SHADER::FDimAtlasType>(NonLocalMeanAtlasType);
+					ComputeShaderPermutationVector.Set<SHADER::FDimPreAlbedoDivide>(GetPreAlbedoDivideRecoverPhase());
+				}
+
+				TShaderMapRef<SHADER> ComputeShader(GetGlobalShaderMap(GMaxRHIFeatureLevel), ComputeShaderPermutationVector);
+
+				FComputeShaderUtils::AddPass(
+					GraphBuilder,
+					RDG_EVENT_NAME("%s ps=%d)",
+						PassInfo.PassName,
+						NonLocalMeanParameters.PatchSize),
+					ERDGPassFlags::Compute,
+					ComputeShader,
+					PassParameters,
+					FComputeShaderUtils::GetGroupCount(SeperableRegionSize, NON_LOCAL_MEAN_THREAD_GROUP_SIZE));
+			}// End filtering
+		} // End dispatch
+
+		// Reshape the buffer from X*Y*Wb to W*X*Y and scatter the results
+		{
+			typedef FNonLocalMeanReshapeBufferCS SHADER;
+			SHADER::FParameters* PassParameters = GraphBuilder.AllocParameters<SHADER::FParameters>();
+			{
+				PassParameters->NLMParams = NonLocalMeanParameters;
+				PassParameters->SourceBuffer = GraphBuilder.CreateSRV(FRDGBufferSRVDesc(NonLocalMeanWeights, PF_G32R32F));
+				PassParameters->RWTargetBuffer = GraphBuilder.CreateUAV(NonLocalMeanWeightsBuffer, PF_R32_FLOAT);
+				PassParameters->SourceBufferDim = FIntVector4(NumSymmetricPairsPerPixel, WeightQueryRegionExtent.X, WeightQueryRegionExtent.Y, TotalNumTilesToFill);
+				PassParameters->TargetBufferDim = FIntVector(NumberOfWeightsPerPixel, Region.Size().X, Region.Size().Y);
+				PassParameters->HalfOffsetSearchCount = HalfOffsetSearchCount;
+			}
+
+			SHADER::FPermutationDomain ComputeShaderPermutationVector;
+
+			TShaderMapRef<SHADER> ComputeShader(GetGlobalShaderMap(GMaxRHIFeatureLevel));
+
+			FComputeShaderUtils::AddPass(
+				GraphBuilder,
+				RDG_EVENT_NAME("NFOR::NonLocalMeanReshapeBufferCS (Size=(%d,%d), HalfOffsetSearchCount=%d)",
+					WeightQueryRegionExtent.X,
+					WeightQueryRegionExtent.Y,
+					HalfOffsetSearchCount),
+				ERDGPassFlags::Compute,
+				ComputeShader,
+				PassParameters,
+				FComputeShaderUtils::GetGroupCount(FIntVector(WeightQueryRegionExtent.X, WeightQueryRegionExtent.Y, HalfOffsetSearchCount), NON_LOCAL_MEAN_THREAD_GROUP_SIZE));
+		}
+	}
+
+	void GetNLMWeights(
+		FRDGBuilder& GraphBuilder,
+		const FSceneView& View,
+		const FRadianceDesc& SourceRadiance,
+		const FRadianceDesc& TargetRadiance,
+		const FRDGBufferRef& NonLocalMeanWeightsBuffer,
+		FIntRect Region,
+		const FNonLocalMeanParameters& NonLocalMeanParameters)
+	{
+		const int32 SearchingPatchSize = (NonLocalMeanParameters.PatchDistance * 2 + 1);
+		const int32 NumberOfWeightsPerPixel = SearchingPatchSize * SearchingPatchSize;
+		const FIntPoint TextureSize = SourceRadiance.Data.Image->Desc.Extent;
+		const bool bSeparateSourceTarget = (SourceRadiance.Data.Image != TargetRadiance.Data.Image);		
+		const bool bShouldNonLocalMeanUseAtlas = ShouldNonLocalMeanUseAtlas();
+
+		RDG_EVENT_SCOPE(GraphBuilder, "NonLocalMeanGetWeights");
+
 		// Query the non-local mean weights for the radiance.
+		if (bShouldNonLocalMeanUseAtlas)
+		{
+			GetNLMWeigthsWithAtlas(GraphBuilder, View, SourceRadiance, TargetRadiance, NonLocalMeanWeightsBuffer, Region, NonLocalMeanParameters);
+		}
+		else
 		{
 			typedef FNonLocalMeanWeightsCS SHADER;
 			SHADER::FParameters* PassParameters = GraphBuilder.AllocParameters<SHADER::FParameters>();
@@ -1174,9 +1507,11 @@ namespace NFORDenoise
 		const FRDGTextureRef& SourceAlbedo
 	)
 	{
+
+		RDG_EVENT_SCOPE(GraphBuilder, "SolveWeightedLSR");
+
 		FIntPoint TextureSize = Radiance->Desc.Extent;
 		
-
 		if (GetRegressionDevice() == ERegressionDevice::CPU)
 		{
 			SolveWeightedLSRCPU(
@@ -1393,6 +1728,8 @@ namespace NFORDenoise
 		FRDGTextureRef RadianceTileTexture = GraphBuilder.CreateTexture(RadianceTileDesc, TEXT("NFOR.RadianceTile"));
 		FRDGTextureRef DenoisedTileTexture = GraphBuilder.CreateTexture(RadianceTileDesc, TEXT("NFOR.DenoisedRadianceTile"));
 		
+		RDG_EVENT_SCOPE(GraphBuilder, "CollaborativeRegression (bandwidth=%.2f)", RadianceNonLocalMeanParameters.Bandwidth);
+
 		AddClearUAVPass(GraphBuilder, GraphBuilder.CreateUAV(FilteredRadiance), FLinearColor::Transparent, ERDGPassFlags::Compute);
 		
 		auto GetTileIndex = [TotalTileCount, NumOfTilesOneSide](int32 Index)
@@ -1421,27 +1758,35 @@ namespace NFORDenoise
 			FIntRect TileRegion = FIntRect(FIntPoint(0), TileSize) + TileStartPoint;
 			FIntRect  PaddedTileRegion = PaddedTileRect + TileStartPoint - PaddingTileOffset;
 
+			RDG_EVENT_SCOPE(GraphBuilder, "Tile (Index=%d)", TileIndex);
+
 			// Get the weights W
-			for (int32 RadianceId = 0; RadianceId < NumOfRadiances; ++RadianceId)
 			{
-				GetNLMWeights(
-					GraphBuilder,
-					View,
-					Radiances[SourceIndex],
-					Radiances[RadianceId],
-					NonLocalMeanSingleFrameWeightsBuffer,
-					TileRegion,
-					RadianceNonLocalMeanParameters);
+				RDG_EVENT_SCOPE(GraphBuilder, "GetNLMWeights (T=%d)", NumOfRadiances);
 
-				if (NumOfRadiances > 1)
+				for (int32 RadianceId = 0; RadianceId < NumOfRadiances; ++RadianceId)
 				{
-					AddCopyBufferPass(GraphBuilder, NonLocalMeanWeightsBuffer, NonLocalMeanSingleFrameWeightSize * BytesPerElement * RadianceId, 
-						NonLocalMeanSingleFrameWeightsBuffer, 0, NonLocalMeanSingleFrameWeightSize * BytesPerElement);
+					GetNLMWeights(
+						GraphBuilder,
+						View,
+						Radiances[SourceIndex],
+						Radiances[RadianceId],
+						NonLocalMeanSingleFrameWeightsBuffer,
+						TileRegion,
+						RadianceNonLocalMeanParameters);
+
+					if (NumOfRadiances > 1)
+					{
+						AddCopyBufferPass(GraphBuilder, NonLocalMeanWeightsBuffer, NonLocalMeanSingleFrameWeightSize * BytesPerElement * RadianceId,
+							NonLocalMeanSingleFrameWeightsBuffer, 0, NonLocalMeanSingleFrameWeightSize * BytesPerElement);
+					}
 				}
 			}
 
 			// Get raw color Y
 			{
+				RDG_EVENT_SCOPE(GraphBuilder, "GetRadiances (T=%d)", NumOfRadiances);
+				
 				int32 BufferChannelOffset = 0;
 				
 				for (int32 RadianceId = 0; RadianceId < NumOfRadiances; ++RadianceId)
@@ -1466,6 +1811,8 @@ namespace NFORDenoise
 
 			// Get the feature vector X
 			{
+				RDG_EVENT_SCOPE(GraphBuilder, "GetFeatureVectors (TxF=%dx%d)", NumOfRadiances, NumOfFeatures / NumOfRadiances);
+
 				int32 BufferChannelOffset = 0;
 
 				for (int32 FeatureId = 0; FeatureId < NumOfFeatures; ++FeatureId)
@@ -1608,7 +1955,7 @@ namespace NFORDenoise
 			FRDGTextureRef Feature = FeatureDesc.Feature.Image;
 			FRDGTextureRef FilterdFeature = GraphBuilder.CreateTexture(Feature->Desc, TEXT("NFOR.FilteredFeature"));
 
-			ApplyNonLocalMeanFilter(
+			ApplyNonLocalMeanFilterIfRequired(
 				GraphBuilder,
 				FeatureNonLocalMeanParameters,
 				FeatureDesc.Feature,
@@ -1659,6 +2006,8 @@ namespace NFORDenoise
 		// Preprocessing
 		// Feature range adjustment, radiance normalization and filtering frames
 		{
+			RDG_EVENT_SCOPE(GraphBuilder, "Preprocessing");
+
 			// Latest frame only 
 			for (int i = 0; i < 1; ++i)
 			{
@@ -1729,7 +2078,7 @@ namespace NFORDenoise
 				FRDGTextureRef FilteredMSETexure = GraphBuilder.CreateTexture(MSE.Image->Desc, TEXT("NFOR.FilteredMSE"));
 				FNonLocalMeanParameters MSENonLocalMeanParameters = GetNonLocalMeanParameters(1, RadiancePatchDistance, 1.0f);
 
-				ApplyNonLocalMeanFilter(
+				ApplyNonLocalMeanFilterIfRequired(
 					GraphBuilder,
 					MSENonLocalMeanParameters,
 					MSE,
@@ -1746,6 +2095,8 @@ namespace NFORDenoise
 		
 		if (bPerformBandwidthSelection)
 		{
+			RDG_EVENT_SCOPE(GraphBuilder, "BandwidthSelection");
+
 			FGlobalShaderMap* GlobalShaderMap = GetGlobalShaderMap(GMaxRHIFeatureLevel);
 			FRDGTextureDesc Desc = FilteredMSEs[0].Image->Desc;
 			FNFORTextureDesc NFORSelectionMap = FilteredMSEs[0];
@@ -1777,7 +2128,7 @@ namespace NFORDenoise
 			{
 				FNonLocalMeanParameters SelectionMapNonLocalMeanParameters = GetNonLocalMeanParameters(1, RadiancePatchDistance, 1.0f);
 
-				ApplyNonLocalMeanFilter(
+				ApplyNonLocalMeanFilterIfRequired(
 					GraphBuilder,
 					SelectionMapNonLocalMeanParameters,
 					NFORSelectionMap,
diff --git a/Engine/Plugins/Experimental/NFORDenoise/Source/NFORDenoise/Private/NFORDenoiseCS.h b/Engine/Plugins/Experimental/NFORDenoise/Source/NFORDenoise/Private/NFORDenoiseCS.h
--- a/Engine/Plugins/Experimental/NFORDenoise/Source/NFORDenoise/Private/NFORDenoiseCS.h
+++ b/Engine/Plugins/Experimental/NFORDenoise/Source/NFORDenoise/Private/NFORDenoiseCS.h
@@ -38,6 +38,13 @@ namespace NFORDenoise
 		MAX
 	};
 
+	enum class ENonLocalMeanAtlasType : uint32
+	{
+		OneSymmetricPair,
+		TwoSymmetricPair,
+		MAX
+	};
+
 	enum class EAlbedoDivideRecoverPhase : uint32
 	{
 		Disabled,
@@ -266,6 +273,14 @@ namespace NFORDenoise
 		SHADER_PARAMETER(float, Bandwidth)
 	END_SHADER_PARAMETER_STRUCT()
 
+	BEGIN_SHADER_PARAMETER_STRUCT(FNonLocalMeanWeightAtlasDispatchParameters, )
+		SHADER_PARAMETER(int32, DispatchId)
+		SHADER_PARAMETER(FIntPoint, DispatchTileSize)
+		SHADER_PARAMETER(int32, DispatchTileCount)
+		SHADER_PARAMETER(FIntRect, SeparableFilteringRegion)
+		SHADER_PARAMETER(FIntVector, DispatchRegionSize)
+	END_SHADER_PARAMETER_STRUCT()
+
 	FNonLocalMeanParameters GetNonLocalMeanParameters(int32 PatchSize, int32 PatchDistance, float Bandwidth);
 
 	// Output the non-local mean filtered image (feature) based on variance
@@ -310,9 +325,6 @@ namespace NFORDenoise
 	//		 Dim: WxHx3, WxHx?
 	// Output: Buffer holding  weights for each pixel
 	//		 Dim: WxHx(2*N+1)^2, where N = NLMParams.PatchDistance
-	//		 TODO?: optimize weights sharing as center(i,j) with offset (dx,dy) has
-	//		 w(i,j,dx,dy)=(i+dx,j+dy,-dx,-dy), 
-
 	class FNonLocalMeanWeightsCS : public FGlobalShader
 	{
 		DECLARE_GLOBAL_SHADER(FNonLocalMeanWeightsCS)
@@ -341,6 +353,90 @@ namespace NFORDenoise
 		using FPermutationDomain = TShaderPermutationDomain<FDimensionVarianceType, FDimensionUseGuide, FDimensionImageChannelCount, FDimensionSeparateSourceTarget, FDimPreAlbedoDivide>;
 	};
 
+	// Optimize the weights query.
+	class FNonLocalMeanGetSqauredDistanceToAtlasCS : public FGlobalShader
+	{
+		DECLARE_GLOBAL_SHADER(FNonLocalMeanGetSqauredDistanceToAtlasCS)
+		SHADER_USE_PARAMETER_STRUCT(FNonLocalMeanGetSqauredDistanceToAtlasCS, FGlobalShader)
+	public:
+		BEGIN_SHADER_PARAMETER_STRUCT(FParameters, )
+			SHADER_PARAMETER_STRUCT_INCLUDE(FNonLocalMeanFilteringCS::FParameters, CommonParameters)
+			SHADER_PARAMETER_STRUCT_INCLUDE(FNonLocalMeanWeightAtlasDispatchParameters, NLMWeightAtlasDispatchParameters)
+			SHADER_PARAMETER_RDG_TEXTURE_SRV(Texture2D, TargetImage)
+			SHADER_PARAMETER_RDG_TEXTURE_SRV(Texture2D, TargetVariance)
+			SHADER_PARAMETER_RDG_TEXTURE_UAV(RWTexture2D, RWNLMWeightAtlas)
+			SHADER_PARAMETER(FIntPoint, NLMWeightAtlasSize)
+		END_SHADER_PARAMETER_STRUCT()
+
+		static void ModifyCompilationEnvironment(const FGlobalShaderPermutationParameters& InParameters, FShaderCompilerEnvironment& OutEnvironment)
+		{
+			FGlobalShader::ModifyCompilationEnvironment(InParameters, OutEnvironment);
+			OutEnvironment.SetDefine(TEXT("THREAD_GROUP_SIZE"), NON_LOCAL_MEAN_THREAD_GROUP_SIZE);
+		}
+
+		class FDimensionVarianceType : SHADER_PERMUTATION_ENUM_CLASS("IMAGE_VARIANCE_TYPE", EVarianceType);
+		class FDimensionImageChannelCount : SHADER_PERMUTATION_RANGE_INT("SOURCE_CHANNEL_COUNT", 1, static_cast<int>(EImageChannelCount::MAX));
+		class FDimensionSeparateSourceTarget : SHADER_PERMUTATION_BOOL("NONLOCALMEAN_SEPARATE_SOURCE");
+		class FDimPreAlbedoDivide : SHADER_PERMUTATION_ENUM_CLASS("PRE_ALBEDO_DIVIDE", EAlbedoDivideRecoverPhase);
+		class FDimAtlasType : SHADER_PERMUTATION_ENUM_CLASS("NONLOCALMEAN_ATLAS_TYPE", ENonLocalMeanAtlasType);
+		using FPermutationDomain = TShaderPermutationDomain<FDimensionVarianceType, FDimensionImageChannelCount, FDimensionSeparateSourceTarget, FDimPreAlbedoDivide, FDimAtlasType>;
+	};
+
+	class FNonLocalMeanSeperableFilterPatchSqauredDistanceCS : public FGlobalShader
+	{
+		DECLARE_GLOBAL_SHADER(FNonLocalMeanSeperableFilterPatchSqauredDistanceCS)
+		SHADER_USE_PARAMETER_STRUCT(FNonLocalMeanSeperableFilterPatchSqauredDistanceCS, FGlobalShader)
+	public:
+		BEGIN_SHADER_PARAMETER_STRUCT(FParameters, )
+			SHADER_PARAMETER_STRUCT_INCLUDE(FNonLocalMeanParameters, NLMParams)
+			SHADER_PARAMETER_STRUCT_INCLUDE(FNonLocalMeanWeightAtlasDispatchParameters, NLMWeightAtlasDispatchParameters)
+			SHADER_PARAMETER_RDG_TEXTURE_SRV(Texture2D, NLMWeightAtlasSource)
+			SHADER_PARAMETER_RDG_TEXTURE_UAV(RWTexture2D, RWNLMWeightAtlasTarget)
+			SHADER_PARAMETER_RDG_BUFFER_UAV(RWBuffer, RWNLMWeights)
+			SHADER_PARAMETER(FIntVector, SeperableRegionSize)
+		END_SHADER_PARAMETER_STRUCT()
+
+		static void ModifyCompilationEnvironment(const FGlobalShaderPermutationParameters& InParameters, FShaderCompilerEnvironment& OutEnvironment)
+		{
+			FGlobalShader::ModifyCompilationEnvironment(InParameters, OutEnvironment);
+			OutEnvironment.SetDefine(TEXT("THREAD_GROUP_SIZE"), NON_LOCAL_MEAN_THREAD_GROUP_SIZE);
+		}
+
+		enum class ESeperablePassType : int32
+		{
+			Horizontal,
+			Vertical,
+			MAX
+		};
+
+		class FDimensionSeperablePassType : SHADER_PERMUTATION_ENUM_CLASS("NONLOCALMEAN_SEPRERABLE_PASS", ESeperablePassType);
+		class FDimPreAlbedoDivide : SHADER_PERMUTATION_ENUM_CLASS("PRE_ALBEDO_DIVIDE", EAlbedoDivideRecoverPhase);
+		class FDimAtlasType : SHADER_PERMUTATION_ENUM_CLASS("NONLOCALMEAN_ATLAS_TYPE", ENonLocalMeanAtlasType);
+		using FPermutationDomain = TShaderPermutationDomain<FDimensionSeperablePassType, FDimPreAlbedoDivide, FDimAtlasType>;
+	};
+
+	// Reshape the layout of the buffer to target
+	// From XxYx[W/B] (each element is of size B) to W*X*Y
+	class FNonLocalMeanReshapeBufferCS : public FGlobalShader
+	{
+		DECLARE_GLOBAL_SHADER(FNonLocalMeanReshapeBufferCS)
+		SHADER_USE_PARAMETER_STRUCT(FNonLocalMeanReshapeBufferCS, FGlobalShader)
+	public:
+		BEGIN_SHADER_PARAMETER_STRUCT(FParameters, )
+			SHADER_PARAMETER_STRUCT_INCLUDE(FNonLocalMeanParameters, NLMParams)
+			SHADER_PARAMETER_RDG_BUFFER_SRV(Buffer<float2>, SourceBuffer)
+			SHADER_PARAMETER_RDG_BUFFER_UAV(RWBuffer<float>, RWTargetBuffer)
+			SHADER_PARAMETER(FIntVector4, SourceBufferDim)
+			SHADER_PARAMETER(FIntVector, TargetBufferDim)
+			SHADER_PARAMETER(int32, HalfOffsetSearchCount)
+		END_SHADER_PARAMETER_STRUCT()
+
+		static void ModifyCompilationEnvironment(const FGlobalShaderPermutationParameters& InParameters, FShaderCompilerEnvironment& OutEnvironment)
+		{
+			FGlobalShader::ModifyCompilationEnvironment(InParameters, OutEnvironment);
+			OutEnvironment.SetDefine(TEXT("THREAD_GROUP_SIZE"), NON_LOCAL_MEAN_THREAD_GROUP_SIZE);
+		}
+	};
 
 	//--------------------------------------------------------------------------------------------------------------------
 	// Collaborative filtering shaders
diff --git a/Engine/Plugins/Experimental/NFORDenoise/Source/NFORDenoise/Shaders/NFORDenoise.usf b/Engine/Plugins/Experimental/NFORDenoise/Source/NFORDenoise/Shaders/NFORDenoise.usf
--- a/Engine/Plugins/Experimental/NFORDenoise/Source/NFORDenoise/Shaders/NFORDenoise.usf
+++ b/Engine/Plugins/Experimental/NFORDenoise/Source/NFORDenoise/Shaders/NFORDenoise.usf
@@ -46,6 +46,30 @@
 #define NONLOCALMEAN_DISTACNE_PATCH 1
 #define NONLOCALMEAN_WEIGHT_TYPE	NONLOCALMEAN_PATCHONLY
 
+#define NONLOCALMEAN_SEPERABLE_FILTER_HORIZONTAL	0
+#define NONLOCALMEAN_SEPERABLE_FILTER_VERTICAL		1
+
+#ifndef NONLOCALMEAN_SEPRERABLE_PASS
+#define NONLOCALMEAN_SEPRERABLE_PASS NONLOCALMEAN_SEPERABLE_FILTER_HORIZONTAL
+#endif
+
+#define NONLOCALMEAN_ATLAS_ONE_SYMMETRIC_PAIR	0
+#define NONLOCALMEAN_ATLAS_TWO_SYMMETRIC_PAIR	1
+
+#ifndef NONLOCALMEAN_ATLAS_TYPE
+#define NONLOCALMEAN_ATLAS_TYPE NONLOCALMEAN_ATLAS_TWO_SYMMETRIC_PAIR
+#endif
+
+#if NONLOCALMEAN_ATLAS_TYPE == NONLOCALMEAN_ATLAS_ONE_SYMMETRIC_PAIR
+#define TAtlasValueType float2
+#elif NONLOCALMEAN_ATLAS_TYPE == NONLOCALMEAN_ATLAS_TWO_SYMMETRIC_PAIR
+#define TAtlasValueType float4
+#else
+#error NONLOCALMEAN_ATLAS_TYPE is not supported.
+#endif
+
+#define NONLOCALMEAN_ATLAS_SYMMETRIC_PAIR_COUNT (NONLOCALMEAN_ATLAS_TYPE+1)
+
 #ifndef NUM_FEATURE
 #define NUM_FEATURE 0
 #endif
@@ -312,6 +336,13 @@ int PatchSize;
 int PatchDistance;
 float Bandwidth;
 
+// Dispatch parameters
+int DispatchId;
+int2 DispatchTileSize;
+int DispatchTileCount;
+int4 SeparableFilteringRegion;
+int3 DispatchRegionSize;
+
 Texture2D<TPixelValue> Image;
 Texture2D<TImageVariance> Variance;
 
@@ -323,37 +354,89 @@ Texture2D<TImageVariance> TargetVariance;
 #define TargetVariance Variance
 #endif
 
-float GetSquaredDistance(int2 P, int2 Q, float VarianceScale = 1.0f)
+struct FImageContext
+{
+	TPixelValue Value;
+	float Variance;
+};
+
+FImageContext GetImageContext(int2 Position, Texture2D<TPixelValue> inImage, Texture2D inVariance, float VarianceScale = 1.0f)
 {
+	FImageContext ImageContext = (FImageContext)0;
+
 	// Use mirrored position. Clamp to boarder has artifacts at border.
-	P = GetMirroredPosition(P, TextureSize); 
-	Q = GetMirroredPosition(Q, TextureSize); 
+	Position = GetMirroredPosition(Position, TextureSize);
 
-	TPixelValue Cp = GetImageValue(P, Image);
-	TPixelValue Cq = GetImageValue(Q, TargetImage);
+	ImageContext.Value = GetImageValue(Position, inImage);
+	ImageContext.Variance = GetImageVariance(Position, inVariance, VarianceChannelOffset) * VarianceScale;
 
-	float VarP = GetImageVariance(P, Variance, VarianceChannelOffset) * VarianceScale;
-	float VarQ = GetImageVariance(Q, TargetVariance, VarianceChannelOffset) * VarianceScale;
+	return ImageContext;
+}
 
-	const float Epsilon =  1e-8f; // Small enough to preserve dark areas
+// Note that GetSquaredDistance(P,Q) != GetSquaredDistance(Q,P), So we cannot share the weight with this asymetric distance.
+float GetSquaredDistance(TPixelValue Cp, float VarP, TPixelValue Cq, float VarQ)
+{
+	const float Epsilon = 1e-8f; // Small enough to preserve dark areas
+
+	float dpq = (Length2(Cp - Cq) - (VarP + min(VarP, VarQ))) / (Epsilon + Pow2(Bandwidth) * (VarP + VarQ));
+
+	const float MaxDistance = 10000;
+	dpq = min(dpq, MaxDistance);
+
+	return dpq;
+}
+
+// Get the both way x=GetSquaredDistance(P,Q) and y=GetSquaredDistance(Q,P) together.
+float2 GetSymmetricSquaredDistance(TPixelValue Cp, float VarP, TPixelValue Cq, float VarQ)
+{
+	const float Epsilon = 1e-8f; // Small enough to preserve dark areas
+
+	float VarPQMin = min(VarP, VarQ);
+	float2 VarPQ = float2((VarP + VarPQMin), (VarQ + VarPQMin));
+
+	float2 dpq = (Length2(Cp - Cq) - VarPQ) / (Epsilon + Pow2(Bandwidth) * (VarP + VarQ));
 
-	float dpq = (Length2(Cp - Cq) - (VarP + min(VarP, VarQ))) / (Epsilon + Pow2(Bandwidth)*(VarP + VarQ));
-	
 	const float MaxDistance = 10000;
 	dpq = min(dpq, MaxDistance);
 
 	return dpq;
 }
 
+float GetSquaredDistance(FImageContext PCtx, FImageContext QCtx)
+{
+	return GetSquaredDistance(PCtx.Value, PCtx.Variance, QCtx.Value, QCtx.Variance);
+}
+
+float GetSquaredDistance(int2 P, int2 Q, float VarianceScale = 1.0f)
+{
+	FImageContext PCtx = GetImageContext(P, Image, Variance, VarianceScale);
+	FImageContext QCtx = GetImageContext(Q, TargetImage, TargetVariance, VarianceScale);
+	float dpq = GetSquaredDistance(PCtx, QCtx);
+	return dpq;
+}
+
+float2 GetSymmetricSquaredDistance(FImageContext PCtx, FImageContext QCtx)
+{
+	return GetSymmetricSquaredDistance(PCtx.Value, PCtx.Variance, QCtx.Value, QCtx.Variance);
+}
+
+float2 GetSymmetricSquaredDistance(int2 P, int2 Q, float VarianceScale = 1.0f)
+{
+	FImageContext PCtx = GetImageContext(P, Image, Variance, VarianceScale);
+	FImageContext QCtx = GetImageContext(Q, TargetImage, TargetVariance, VarianceScale);
+	float2 dpq = GetSymmetricSquaredDistance(PCtx, QCtx);
+	return dpq;
+}
+
 float GetPatchSquaredDistance(int2 P, int2 Q)
 {
 	const float VarianceScale = 1.0f;
 
 	// TODO: z-order for performance improvement?
 	float Sum = 0;
-	for (int y = -PatchSize; y < PatchSize; ++y)
+	for (int y = -PatchSize; y <= PatchSize; ++y)
 	{
-		for (int x = -PatchSize; x < PatchSize; ++x)
+		for (int x = -PatchSize; x <= PatchSize; ++x)
 		{
 			int2 Pij = P + int2(x, y);
 			int2 Qij = Q + int2(x, y);
@@ -366,6 +449,25 @@ float GetPatchSquaredDistance(int2 P, int2 Q)
 	return PathSquareDistance;
 }
 
+TAtlasValueType GetNonLocalMeanWeight(int2 P, int2 Q, TAtlasValueType PatchPQSqaureDistance)
+{
+	TAtlasValueType NonLocalMeanWeight = (TAtlasValueType)0;
+#if NONLOCALMEAN_WEIGHT_TYPE == NONLOCALMEAN_PATCHONLY
+	NonLocalMeanWeight = FastExp(-PatchPQSqaureDistance);
+#elif NONLOCALMEAN_WEIGHT_TYPE == NONLOCALMEAN_DISTACNE_PATCH
+	#error not implemented
+	const float SigmaS = 2.0f;//TODO: move to CVar
+	NonLocalMeanWeight = FastExp(-length2(P - Q) / (2 * Pow2(SigmaS))) * FastExp(-PatchPQSqaureDistance);
+#endif
+
+#if PRE_ALBEDO_DIVIDE == PRE_ALBEDO_DIVIDE_DISABLED
+	const TAtlasValueType kMinFloat16 = (TAtlasValueType)0.000000059604645f;
+	//without albedo divide, below is better when scattering results.
+	NonLocalMeanWeight = max(kMinFloat16, NonLocalMeanWeight);
+#endif
+	return NonLocalMeanWeight;
+}
+
 // P is on source image, Q is on target image
 float GetNonLocalMeanWeight(int2 P, int2 Q)
 {
@@ -454,6 +556,176 @@ void NonLocalMeanWeightsCS(in const uint3 DispatchThreadID : SV_DispatchThreadID
 	}
 }
 
+RWTexture2D<TAtlasValueType> RWNLMWeightAtlas;
+int2 NLMWeightAtlasSize;
+
+uint2 GetAtlasPosition(uint3 DispatchThreadID, uint2 TileRegionSize,uint2 Offset = 0)
+{
+	uint2 TileId = uint2(DispatchThreadID.z % DispatchTileSize.x, DispatchThreadID.z / DispatchTileSize.x);
+	uint2 TileOffset = TileId * TileRegionSize;
+	uint2 AtlasPosition = TileOffset + DispatchThreadID.xy + Offset;
+
+	return AtlasPosition;
+}
+
+[numthreads(THREAD_GROUP_SIZE, THREAD_GROUP_SIZE, THREAD_GROUP_SIZE)]
+void NonLocalMeanGetSqauredDistanceToAtlasCS(in const uint3 DispatchThreadID : SV_DispatchThreadID)
+{
+	if (any(DispatchThreadID >= DispatchRegionSize))
+	{
+		return;
+	}
+
+	const float VarianceScale = 1.0f;
+
+	// Query the current position.
+	uint2 Position = DispatchThreadID.xy + SeparableFilteringRegion.xy;
+	FImageContext PContext = GetImageContext(Position, Image, Variance, VarianceScale);
+
+	// Calculate the symmetric squared distance.
+	int BaseOffsetIndex = (DispatchId * DispatchTileCount + DispatchThreadID.z) * NONLOCALMEAN_ATLAS_SYMMETRIC_PAIR_COUNT;
+	int PatchSearchWidth = PatchDistance * 2 + 1;
+
+	TAtlasValueType Distances = (TAtlasValueType)0;
+	FImageContext QContext;
+
+	UNROLL
+	for (int i = 0; i < NONLOCALMEAN_ATLAS_SYMMETRIC_PAIR_COUNT; ++i)
+	{
+		int OffsetIndex = BaseOffsetIndex + i;
+
+		int2 Offset = int2(OffsetIndex % PatchSearchWidth, OffsetIndex / PatchSearchWidth) - int2(PatchDistance, PatchDistance);
+		int2 Q = Position + Offset;
+
+		QContext = GetImageContext(Q, TargetImage, TargetVariance, VarianceScale);
+		float2 Distance = GetSymmetricSquaredDistance(PContext, QContext);
+		Distances[i * 2 + 0] = Distance.x;
+		Distances[i * 2 + 1] = Distance.y;
+	}
+
+	uint2 AtlasPosition = GetAtlasPosition(DispatchThreadID, DispatchRegionSize.xy);
+	RWNLMWeightAtlas[AtlasPosition] = Distances;
+}
+
+RWBuffer<TAtlasValueType> RWNLMWeights;
+Texture2D<TAtlasValueType> NLMWeightAtlasSource;
+RWTexture2D<TAtlasValueType> RWNLMWeightAtlasTarget;
+int3 SeperableRegionSize;
+
+[numthreads(THREAD_GROUP_SIZE, THREAD_GROUP_SIZE, THREAD_GROUP_SIZE)]
+void NonLocalMeanSeperableFilterPatchSqauredDistanceCS(in const uint3 DispatchThreadID : SV_DispatchThreadID)
+{
+	if (any(DispatchThreadID >= SeperableRegionSize))
+	{
+		return;
+	}
+
+#if NONLOCALMEAN_SEPRERABLE_PASS == NONLOCALMEAN_SEPERABLE_FILTER_HORIZONTAL
+	int2 Direction = int2(1, 0);
+	uint2 Offset = 0;
+#elif NONLOCALMEAN_SEPRERABLE_PASS == NONLOCALMEAN_SEPERABLE_FILTER_VERTICAL
+	int2 Direction = int2(0, 1);
+	uint2 Offset = uint2(PatchSize,PatchSize); // Don't need to run the vertical filter  for the outside patch region.
+#endif
+
+	uint2 AtlasPosition = GetAtlasPosition(DispatchThreadID, DispatchRegionSize.xy, Offset);
+
+	TAtlasValueType Sum = (TAtlasValueType)0;
+	for (int i = -PatchSize; i <= PatchSize; ++i)
+	{
+		int2 Pij = AtlasPosition + i * Direction;
+		Sum += NLMWeightAtlasSource[Pij];
+	}
+
+	Sum /= (2 * PatchSize + 1);
+
+#if NONLOCALMEAN_SEPRERABLE_PASS == NONLOCALMEAN_SEPERABLE_FILTER_HORIZONTAL
+	RWNLMWeightAtlasTarget[AtlasPosition] = Sum;
+#elif NONLOCALMEAN_SEPRERABLE_PASS == NONLOCALMEAN_SEPERABLE_FILTER_VERTICAL
+	uint2 Q = 0;
+	TAtlasValueType PatchSquareDistance = max(0, Sum);
+	TAtlasValueType PatchWeight0 = GetNonLocalMeanWeight(AtlasPosition, Q, PatchSquareDistance);
+	
+	// Write to buffer with dimension layout XxYxNumB (each element is of size Count(TAtlasValueType,float)/2)
+	const uint x = DispatchThreadID.x;
+	const uint X = SeperableRegionSize.x;
+	const uint y = DispatchThreadID.y;
+	const uint Y = SeperableRegionSize.y;
+	const uint z = DispatchId * DispatchTileCount + DispatchThreadID.z;
+
+	const uint SaveIndex =  x + X * (y + z * Y);
+	RWNLMWeights[SaveIndex] = PatchWeight0;
+#endif
+}
+
+Buffer<float2> SourceBuffer;
+RWBuffer<float> RWTargetBuffer;
+int4 SourceBufferDim;// B,X+2*pd, Y+2*pd, Wb
+int3 TargetBufferDim;// W, X, Y
+int HalfOffsetSearchCount;
+
+float2 GetInputFromBuffer(uint x, uint y, uint w)
+{
+	const uint B = SourceBufferDim.x;
+	const uint X = SourceBufferDim.y;
+	const uint Y = SourceBufferDim.z;
+	const uint Wb = SourceBufferDim.w;
+
+	const uint b = w % B;
+	const uint wb = w / B;
+
+	const uint InputIndex = b + B * (x + X * (y + Y * wb));
+	return SourceBuffer[InputIndex];
+}
+
+void WriteWeight(uint2 PixelPosition, uint2 WeightPosition, float Weight)
+{
+	const uint PatchDiemeter = 2 * PatchDistance + 1;
+	uint w = WeightPosition.x + WeightPosition.y * PatchDiemeter;
+	uint OutputIndex = w + PixelPosition.x * TargetBufferDim.x + PixelPosition.y * TargetBufferDim.x * TargetBufferDim.y;
+	RWTargetBuffer[OutputIndex] = Weight;
+}
+
+[numthreads(THREAD_GROUP_SIZE, THREAD_GROUP_SIZE, THREAD_GROUP_SIZE)]
+void NonLocalMeanReshapeBufferCS(in const uint3 DispatchThreadID : SV_DispatchThreadID)
+{
+	uint3 ReshapeDim = uint3(SourceBufferDim.y, SourceBufferDim.z, HalfOffsetSearchCount);
+	if (any(DispatchThreadID >= ReshapeDim.xyz))
+	{
+		return;
+	}
+	
+	const uint x2pd = DispatchThreadID.x;
+	const uint y2pd = DispatchThreadID.y;
+	const uint w	= DispatchThreadID.z;
+
+	// Get from input
+	float2 PatchWeights = GetInputFromBuffer(x2pd, y2pd, w);// .x stores (x,y)->P, weights
+															// .y stores P ->(x,y) weights
+	// For a given w index in upper matrix, convert to (i,j), find the mapping to lower coordinate.
+	const uint PatchWidth = 2 * PatchDistance + 1;
+	const int2 WeightCenterPoint = int2(PatchDistance, PatchDistance);
+	int2 WeightCoordinateP = int2 (w % PatchWidth, w / PatchWidth);
+	int2 DeltaXY = WeightCoordinateP - WeightCenterPoint;
+	int2 WeightCoordinateQ = WeightCenterPoint - DeltaXY;
+	int2 CoordinateXY2pd = int2(x2pd,y2pd);
+	int2 CoordinateP2pd = CoordinateXY2pd + DeltaXY;
+	
+	// Save (x,y) -> P weights
+	int2 RegionSize = TargetBufferDim.yz;
+	int2 CoordinateXY = CoordinateXY2pd - PatchDistance;
+	if (all(CoordinateXY < RegionSize) && all(CoordinateXY >= 0))
+	{
+		WriteWeight(CoordinateXY, WeightCoordinateP, PatchWeights.x);
+	}
+	// Save P -> (x,y) weights if it's inside the region.
+	int2 CoordinateP = CoordinateP2pd - PatchDistance;
+	if (all(CoordinateP < RegionSize) && all (CoordinateP >= 0))
+	{
+		WriteWeight(CoordinateP, WeightCoordinateQ, PatchWeights.y);
+	}
+}
+
 //------------------------------------------------------------------------------------------------------
 // Collaborative filtering
 //	1. Tiling
