Imporve the Linear solver efficiency for AX=B where A is symmetric and positive semidefinite for NFOR denoiser. Use Cholesky to get to the neighbor, and apply a few more Newton iterations to refine the result. Performance is improved for 10 times.

diff --git a/Engine/Plugins/Experimental/NFORDenoise/Source/NFORDenoise/Private/NFORDenoiseCS.cpp b/Engine/Plugins/Experimental/NFORDenoise/Source/NFORDenoise/Private/NFORDenoiseCS.cpp
--- a/Engine/Plugins/Experimental/NFORDenoise/Source/NFORDenoise/Private/NFORDenoiseCS.cpp
+++ b/Engine/Plugins/Experimental/NFORDenoise/Source/NFORDenoise/Private/NFORDenoiseCS.cpp
@@ -117,6 +117,24 @@ namespace NFORDenoise
 		TEXT("1: Solve Ax=B on GPU."),
 		ECVF_RenderThreadSafe);
 
+	TAutoConsoleVariable<int32> CVarNFORLinearSolverType(
+		TEXT("r.NFOR.LinearSolver.Type"),
+		2,
+		TEXT("The linear regression solver type implemented in GPU.\n")
+		TEXT("0: Newton Schulz iterative method (High quality but slow).\n")
+		TEXT("1: Cholesky decomposition (Fast but has too smoothed result or artifacts).\n")
+		TEXT("2: Fusion of Cholesky and Newton Schulz iterative method (High quality and fast).\n"),
+		ECVF_RenderThreadSafe);
+
+	TAutoConsoleVariable<float> CVarNFORLinearSolverCholeskyLambda(
+		TEXT("r.NFOR.LinearSolver.Cholesky.Lambda"),
+		1e-3,
+		TEXT("The parameter lambda for modified Cholesky decomposition to make it positive definite.\n")
+		TEXT("Large value yields bias with smoothed rendering, while small value leads to variance or artifacts.\n")
+		TEXT("Used when r.NFOR.LinearSolver.Type = 1 and 2. Selected to match the quality of r.NFOR.LinearSolver.Type = 2\n")
+		TEXT("to r.NFOR.LinearSolver.Type 0."),
+		ECVF_RenderThreadSafe);
+
 	TAutoConsoleVariable<int32> CVarNFORReconstructionType(
 		TEXT("r.NFOR.Reconstruction.Type"),
 		0,
@@ -344,6 +362,32 @@ namespace NFORDenoise
 		return static_cast<ELinearSolverDevice>(LinearSolverDevice);
 	}
 
+	RegressionKernel::FLinearSolverCS::ESolverType GetLinearSolverType()
+	{
+		const int32 LinearSolverType = FMath::Clamp(CVarNFORLinearSolverType.GetValueOnRenderThread(),
+			0,
+			static_cast<int32>(RegressionKernel::FLinearSolverCS::ESolverType::MAX));
+		// MAX indicates using fusion.
+		return static_cast<RegressionKernel::FLinearSolverCS::ESolverType>(LinearSolverType);
+	}
+
+	const TCHAR* GetLinearSolverTypeName(RegressionKernel::FLinearSolverCS::ESolverType& SolverType)
+	{
+		static const TCHAR* const kEventNames[] = {
+				TEXT("NewtonSchulz"),
+				TEXT("Cholesky"),
+				TEXT("NewtonCholesky"),
+				TEXT("Fusion")
+		};
+		static_assert(UE_ARRAY_COUNT(kEventNames) == (int32(RegressionKernel::FLinearSolverCS::ESolverType::MAX)+1), "Fix me");
+		return kEventNames[int32(SolverType)];
+	}
+
+	float GetLinearSolverCholeskyLambda()
+	{
+		return FMath::Max(CVarNFORLinearSolverCholeskyLambda.GetValueOnRenderThread(),0);
+	}
+
 	RegressionKernel::FReconstructSpatialTemporalImage::EReconstructionType 
 		GetReconstructionType( int32 CurrentFrameIndex, int32 DenoisingFrameIndex)
 	{
@@ -488,6 +532,12 @@ namespace NFORDenoise
 	//	2. Weighted Least-square solver
 	IMPLEMENT_GLOBAL_SHADER(RegressionKernel::FInPlaceBatchedMatrixMultiplicationCS, "/NFORDenoise/NFORDenoise.usf", "InPlaceBatchedMatrixMultiplicationCS", SF_Compute);
 	IMPLEMENT_GLOBAL_SHADER(RegressionKernel::FLinearSolverCS, "/NFORDenoise/NFORDenoise.usf", "LinearSolverCS", SF_Compute);
+	
+	//		Allow quality and speed balance
+	IMPLEMENT_GLOBAL_SHADER(RegressionKernel::FLinearSolverBuildIndirectDispatchArgsCS, "/NFORDenoise/NFORDenoise.usf", "LinearSolverBuildIndirectDispatchArgsCS", SF_Compute);
+	IMPLEMENT_GLOBAL_SHADER(RegressionKernel::FLinearSolverIndirectCS, "/NFORDenoise/NFORDenoise.usf", "LinearSolverIndirectCS", SF_Compute);
+
+	
 	IMPLEMENT_GLOBAL_SHADER(RegressionKernel::FReconstructSpatialTemporalImage, "/NFORDenoise/NFORDenoise.usf", "ReconstructSpatialTemporalImageCS", SF_Compute);
 	IMPLEMENT_GLOBAL_SHADER(FAccumulateBufferToTextureCS, "/NFORDenoise/NFORDenoise.usf", "AccumulateBufferToTextureCS", SF_Compute);
 
@@ -1495,6 +1545,311 @@ namespace NFORDenoise
 		}
 	}
 
+	void ApplyLinearSolverGPU(
+		FRDGBuilder& GraphBuilder,
+		FRDGBufferRef AMatrix,
+		FRDGBufferRef BMatrix,
+		FIntPoint BDim,
+		FRDGBufferRef ReconstructionWeights,
+		const FWeightedLSRDesc& WeightedLSRDesc,
+		int32 TotalNumOfFeaturesPerFrame,
+		int32 NumOfElements,
+		int32 NumOfElementsPerRow
+	)
+	{
+
+		RDG_EVENT_SCOPE(GraphBuilder, "BatchedLinearSolver");
+
+		// Summary of approximate ground truth solver
+		// 1. Apply Cholesky decomposition with lambda = 0. and output failed indices.
+		// 2. Apply Cholesky decomposition with lambda = 1e-6 on failed, and output both failed and succeeded indices.
+		// 3. For failed indices, fallback to newton iterative method.
+		// 4. For succeeded indices, iteratively refine with better lambda.
+		// Summary of NewtonCholesky
+		// 1. Apply Cholesky decomposition with lambda = 1e-3, fine tune with 3 iterations of Newton. If any inversion failed, output the failed indices.
+		// 2. For failed indices, apply the standard newton iteration method. 
+		// Other wise, solve based on the SolverType in a single pass.
+
+		RegressionKernel::FLinearSolverCS::ESolverType SolverType = GetLinearSolverType();
+		const bool bApproximateGroundTruthSolver = SolverType == RegressionKernel::FLinearSolverCS::ESolverType::MAX;
+		const bool bUseSuccessAndFailIndexBuffer = bApproximateGroundTruthSolver || SolverType == RegressionKernel::FLinearSolverCS::ESolverType::NewtonCholesky;
+		// Success count | Failed count | sidx... <--->     fidx|
+		// One for read, one for write
+		FRDGBufferRef SuccessAndFailIndexBuffer[2] = { nullptr, nullptr};
+		if (bUseSuccessAndFailIndexBuffer)
+		{
+			const int32 BytesPerElement = sizeof(uint32);
+			FRDGBufferDesc SuccessAndFailIndexBufferDesc = FRDGBufferDesc::CreateBufferDesc(BytesPerElement, NumOfElements + 2);
+			SuccessAndFailIndexBuffer[0] = GraphBuilder.CreateBuffer(SuccessAndFailIndexBufferDesc, TEXT("NFOR.LinearSolver.SuccessAndFailIndexBuffer0"));
+			
+			// Initialize the first two elements to 0 for SuccessAndFailIndexBuffer.
+			FRDGBufferDesc IndicesHeadBufferDesc =
+				FRDGBufferDesc::CreateBufferDesc(BytesPerElement, 2);
+			FRDGBufferRef IndicesHeadBuffer = GraphBuilder.CreateBuffer(IndicesHeadBufferDesc, TEXT("NFOR.LinearSolver.IndexHeadBuffer"));
+
+			AddClearUAVPass(GraphBuilder, GraphBuilder.CreateUAV(IndicesHeadBuffer, PF_R32_UINT), 0);
+			AddCopyBufferPass(GraphBuilder, SuccessAndFailIndexBuffer[0], 0, IndicesHeadBuffer, 0, BytesPerElement * 2);
+			
+			if (bApproximateGroundTruthSolver)
+			{
+				SuccessAndFailIndexBuffer[1] = GraphBuilder.CreateBuffer(SuccessAndFailIndexBufferDesc, TEXT("NFOR.LinearSolver.SuccessAndFailIndexBuffer1"));
+				AddCopyBufferPass(GraphBuilder, SuccessAndFailIndexBuffer[1], 0, IndicesHeadBuffer, 0, BytesPerElement * 2);
+			}
+		}
+
+		RegressionKernel::FLinearSolverCS::FParameters CommonPassParameters;
+		{
+			CommonPassParameters.A = GraphBuilder.CreateSRV(FRDGBufferSRVDesc(AMatrix, PF_R32_FLOAT));
+			CommonPassParameters.ADim = FIntPoint(TotalNumOfFeaturesPerFrame, TotalNumOfFeaturesPerFrame);
+			CommonPassParameters.B = GraphBuilder.CreateSRV(FRDGBufferSRVDesc(BMatrix, PF_R32_FLOAT));
+			CommonPassParameters.BDim = BDim;
+			CommonPassParameters.Result = GraphBuilder.CreateUAV(FRDGBufferUAVDesc(ReconstructionWeights, PF_R32_FLOAT));
+			CommonPassParameters.NumOfElements = NumOfElements;
+			CommonPassParameters.NumOfElementsPerRow = NumOfElementsPerRow;
+			CommonPassParameters.Lambda = 0.0f;
+		}
+
+		// First multi-pass or the single pass based on SolverType.
+		{
+			typedef RegressionKernel::FLinearSolverCS SHADER;
+			SHADER::ESolverType FirstPassSolverType = bApproximateGroundTruthSolver ? SHADER::ESolverType::Cholesky : SolverType;
+			SHADER::FParameters* PassParameters = GraphBuilder.AllocParameters<SHADER::FParameters>();
+			{
+				*PassParameters = CommonPassParameters;
+				//Magnitude of X^TWX element value increases with the number of frames, and the number of elements selected to 
+				// estimate the weights. GetLinearSolverCholeskyLambda() returns the lambda for a single frame.
+				PassParameters->Lambda = GetLinearSolverCholeskyLambda() * WeightedLSRDesc.NumOfFrames;
+				if (bApproximateGroundTruthSolver)
+				{
+					PassParameters->RWSuccessAndFailIndexBuffer = GraphBuilder.CreateUAV(FRDGBufferUAVDesc(SuccessAndFailIndexBuffer[0], PF_R32_UINT));
+					PassParameters->Lambda = 0.0f;
+				}
+			}
+
+			SHADER::FPermutationDomain ComputeShaderPermutationVector;
+			{
+				checkf(BDim.X >= 6 && BDim.X <= 8, TEXT("Number of features should be between 6 and 8"));
+				ComputeShaderPermutationVector.Set<SHADER::FDimNumFeature>(BDim.X);
+				ComputeShaderPermutationVector.Set<SHADER::FDimSolverType>(FirstPassSolverType);
+				ComputeShaderPermutationVector.Set<SHADER::FDimOutputIndices>(bApproximateGroundTruthSolver);
+			}
+
+			TShaderMapRef<SHADER> ComputeShader(GetGlobalShaderMap(GMaxRHIFeatureLevel), ComputeShaderPermutationVector);
+
+			FComputeShaderUtils::AddPass(
+				GraphBuilder,
+				RDG_EVENT_NAME("NFOR::BatchedLinearSolver(F=%d, C=%d, %s)", 
+					BDim.X, BDim.Y, 
+					GetLinearSolverTypeName(SolverType)),
+				ERDGPassFlags::Compute,
+				ComputeShader,
+				PassParameters,
+				FComputeShaderUtils::GetGroupCount(FIntPoint(WeightedLSRDesc.Width, WeightedLSRDesc.Height), NON_LOCAL_MEAN_THREAD_GROUP_SIZE));
+		}
+
+		// Early out for Newton or Cholesky method.
+		if (!bUseSuccessAndFailIndexBuffer)
+		{
+			return;
+		}
+
+		FRDGBufferRef IndirectDispatchArgsBuffer = 
+			GraphBuilder.CreateBuffer(FRDGBufferDesc::CreateIndirectDesc<FRHIDispatchIndirectParameters>(1), TEXT("NFOR.LinearSolver.IndirectDispatchBuffer"));
+
+		{
+			{
+				// Build the indirect dispatch parameters on failed
+				typedef RegressionKernel::FLinearSolverBuildIndirectDispatchArgsCS SHADER;
+				SHADER::FParameters* PassParameters = GraphBuilder.AllocParameters<SHADER::FParameters>();
+				{
+					PassParameters->SuccessAndFailIndexBuffer = GraphBuilder.CreateSRV(FRDGBufferSRVDesc(SuccessAndFailIndexBuffer[0], PF_R32_UINT));
+					PassParameters->RWIndirectDispatchArgsBuffer = GraphBuilder.CreateUAV(FRDGBufferUAVDesc(IndirectDispatchArgsBuffer, PF_R32_UINT));
+				}
+
+				SHADER::FPermutationDomain ComputeShaderPermutationVector;
+				{
+					ComputeShaderPermutationVector.Set<SHADER::FDimInputMatrixType>(RegressionKernel::EInputMatrixType::Fail);
+				}
+
+				TShaderMapRef<SHADER> ComputeShader(GetGlobalShaderMap(GMaxRHIFeatureLevel), ComputeShaderPermutationVector);
+
+				FComputeShaderUtils::AddPass(
+					GraphBuilder,
+					RDG_EVENT_NAME("NFOR::BuildIndirectDispatchCS"),
+					ERDGPassFlags::Compute,
+					ComputeShader,
+					PassParameters,
+					FIntVector(1, 1, 1));
+			}
+
+			typedef RegressionKernel::FLinearSolverIndirectCS SHADER;
+			RegressionKernel::FLinearSolverCS::ESolverType PassSolverType = bApproximateGroundTruthSolver ?
+				RegressionKernel::FLinearSolverCS::ESolverType::Cholesky : RegressionKernel::FLinearSolverCS::ESolverType::NewtonSchulz;
+			const float LambdaExponent = -6.0f;
+			SHADER::FParameters* PassParameters = GraphBuilder.AllocParameters<SHADER::FParameters>();
+			{
+				PassParameters->CommonParameters = CommonPassParameters;
+				PassParameters->CommonParameters.Lambda = FMath::Pow(10, LambdaExponent);
+				if (bApproximateGroundTruthSolver)
+				{
+					PassParameters->CommonParameters.RWSuccessAndFailIndexBuffer = GraphBuilder.CreateUAV(FRDGBufferUAVDesc(SuccessAndFailIndexBuffer[1], PF_R32_UINT));
+				}
+				PassParameters->SuccessAndFailIndexBuffer = GraphBuilder.CreateSRV(FRDGBufferSRVDesc(SuccessAndFailIndexBuffer[0], PF_R32_UINT));
+
+				PassParameters->IndirectDispatchArgsBuffer = IndirectDispatchArgsBuffer;
+			}
+
+			SHADER::FPermutationDomain ComputeShaderPermutationVector;
+			{
+				checkf(BDim.X >= 6 && BDim.X <= 8, TEXT("Number of features should be between 6 and 8"));
+				ComputeShaderPermutationVector.Set<SHADER::FDimNumFeature>(BDim.X);
+				ComputeShaderPermutationVector.Set<SHADER::FDimSolverType>(PassSolverType);
+				ComputeShaderPermutationVector.Set<SHADER::FDimInputMatrixType>(RegressionKernel::EInputMatrixType::Fail);
+				ComputeShaderPermutationVector.Set<SHADER::FDimOutputIndices>(bApproximateGroundTruthSolver);
+			}
+
+			TShaderMapRef<SHADER> ComputeShader(GetGlobalShaderMap(GMaxRHIFeatureLevel), ComputeShaderPermutationVector);
+
+			FComputeShaderUtils::AddPass(
+				GraphBuilder,
+				RDG_EVENT_NAME("NFOR::BatchedLinearSolverIndirect(%s, Lambda=1e%.1f)", 
+					GetLinearSolverTypeName(PassSolverType),
+					LambdaExponent),
+				ERDGPassFlags::Compute,
+				ComputeShader,
+				PassParameters,
+				PassParameters->IndirectDispatchArgsBuffer, 0);
+
+		}
+
+		// Early out for Newton Cholesky method.
+		if (!bApproximateGroundTruthSolver)
+		{
+			return;
+		}
+
+		{
+			//For failed indices, fallback to newton iterative method.
+			{
+				// Build the indirect dispatch parameters on failed
+				typedef RegressionKernel::FLinearSolverBuildIndirectDispatchArgsCS SHADER;
+				SHADER::FParameters* PassParameters = GraphBuilder.AllocParameters<SHADER::FParameters>();
+				{
+					PassParameters->SuccessAndFailIndexBuffer = GraphBuilder.CreateSRV(FRDGBufferSRVDesc(SuccessAndFailIndexBuffer[1], PF_R32_UINT));
+					PassParameters->RWIndirectDispatchArgsBuffer = GraphBuilder.CreateUAV(FRDGBufferUAVDesc(IndirectDispatchArgsBuffer, PF_R32_UINT));
+				}
+
+				SHADER::FPermutationDomain ComputeShaderPermutationVector;
+				{
+					ComputeShaderPermutationVector.Set<SHADER::FDimInputMatrixType>(RegressionKernel::EInputMatrixType::Fail);
+				}
+
+				TShaderMapRef<SHADER> ComputeShader(GetGlobalShaderMap(GMaxRHIFeatureLevel), ComputeShaderPermutationVector);
+
+				FComputeShaderUtils::AddPass(
+					GraphBuilder,
+					RDG_EVENT_NAME("NFOR::BuildIndirectDispatchCS"),
+					ERDGPassFlags::Compute,
+					ComputeShader,
+					PassParameters,
+					FIntVector(1, 1, 1));
+			}
+
+			typedef RegressionKernel::FLinearSolverIndirectCS SHADER;
+			RegressionKernel::FLinearSolverCS::ESolverType PassSolverType = RegressionKernel::FLinearSolverCS::ESolverType::NewtonSchulz;
+			SHADER::FParameters* PassParameters = GraphBuilder.AllocParameters<SHADER::FParameters>();
+			{
+				PassParameters->CommonParameters = CommonPassParameters;
+				PassParameters->SuccessAndFailIndexBuffer = GraphBuilder.CreateSRV(FRDGBufferSRVDesc(SuccessAndFailIndexBuffer[1], PF_R32_UINT));
+
+				PassParameters->IndirectDispatchArgsBuffer = IndirectDispatchArgsBuffer;
+			}
+
+			SHADER::FPermutationDomain ComputeShaderPermutationVector;
+			{
+				checkf(BDim.X >= 6 && BDim.X <= 8, TEXT("Number of features should be between 6 and 8"));
+				ComputeShaderPermutationVector.Set<SHADER::FDimNumFeature>(BDim.X);
+				ComputeShaderPermutationVector.Set<SHADER::FDimSolverType>(PassSolverType);
+				ComputeShaderPermutationVector.Set<SHADER::FDimInputMatrixType>(RegressionKernel::EInputMatrixType::Fail);
+				ComputeShaderPermutationVector.Set<SHADER::FDimOutputIndices>(false);
+			}
+
+			TShaderMapRef<SHADER> ComputeShader(GetGlobalShaderMap(GMaxRHIFeatureLevel), ComputeShaderPermutationVector);
+
+			FComputeShaderUtils::AddPass(
+				GraphBuilder,
+				RDG_EVENT_NAME("NFOR::BatchedLinearSolverIndirect(%s on Failed)",
+					GetLinearSolverTypeName(PassSolverType)),
+				ERDGPassFlags::Compute,
+				ComputeShader,
+				PassParameters,
+				PassParameters->IndirectDispatchArgsBuffer, 0);
+		}
+
+		{
+			// For succeeded indices, iteratively refine with smaller lambda.
+			// Lambda = 10-7. TODO: ieratively refine.
+
+			{
+				// Build the indirect dispatch parameters on failed
+				typedef RegressionKernel::FLinearSolverBuildIndirectDispatchArgsCS SHADER;
+				SHADER::FParameters* PassParameters = GraphBuilder.AllocParameters<SHADER::FParameters>();
+				{
+					PassParameters->SuccessAndFailIndexBuffer = GraphBuilder.CreateSRV(FRDGBufferSRVDesc(SuccessAndFailIndexBuffer[1], PF_R32_UINT));
+					PassParameters->RWIndirectDispatchArgsBuffer = GraphBuilder.CreateUAV(FRDGBufferUAVDesc(IndirectDispatchArgsBuffer, PF_R32_UINT));
+				}
+
+				SHADER::FPermutationDomain ComputeShaderPermutationVector;
+				{
+					ComputeShaderPermutationVector.Set<SHADER::FDimInputMatrixType>(RegressionKernel::EInputMatrixType::Success);
+				}
+
+				TShaderMapRef<SHADER> ComputeShader(GetGlobalShaderMap(GMaxRHIFeatureLevel), ComputeShaderPermutationVector);
+
+				FComputeShaderUtils::AddPass(
+					GraphBuilder,
+					RDG_EVENT_NAME("NFOR::BuildIndirectDispatchCS"),
+					ERDGPassFlags::Compute,
+					ComputeShader,
+					PassParameters,
+					FIntVector(1, 1, 1));
+			}
+
+			const float LambdaExponent = -7.0f;
+			typedef RegressionKernel::FLinearSolverIndirectCS SHADER;
+			RegressionKernel::FLinearSolverCS::ESolverType PassSolverType = RegressionKernel::FLinearSolverCS::ESolverType::Cholesky;
+			SHADER::FParameters* PassParameters = GraphBuilder.AllocParameters<SHADER::FParameters>();
+			{
+				PassParameters->CommonParameters = CommonPassParameters;
+				PassParameters->SuccessAndFailIndexBuffer = GraphBuilder.CreateSRV(FRDGBufferSRVDesc(SuccessAndFailIndexBuffer[1], PF_R32_UINT));
+				PassParameters->IndirectDispatchArgsBuffer = IndirectDispatchArgsBuffer;
+				PassParameters->CommonParameters.Lambda = FMath::Pow(10, LambdaExponent);
+			}
+
+			SHADER::FPermutationDomain ComputeShaderPermutationVector;
+			{
+				checkf(BDim.X >= 6 && BDim.X <= 8, TEXT("Number of features should be between 6 and 8"));
+				ComputeShaderPermutationVector.Set<SHADER::FDimNumFeature>(BDim.X);
+				ComputeShaderPermutationVector.Set<SHADER::FDimSolverType>(PassSolverType);
+				ComputeShaderPermutationVector.Set<SHADER::FDimInputMatrixType>(RegressionKernel::EInputMatrixType::Success);
+				ComputeShaderPermutationVector.Set<SHADER::FDimOutputIndices>(false);
+			}
+
+			TShaderMapRef<SHADER> ComputeShader(GetGlobalShaderMap(GMaxRHIFeatureLevel), ComputeShaderPermutationVector);
+
+			FComputeShaderUtils::AddPass(
+				GraphBuilder,
+				RDG_EVENT_NAME("NFOR::BatchedLinearSolverIndirect(%s on Succeeded, Lambda=1e%.1f)",
+					GetLinearSolverTypeName(PassSolverType),
+					LambdaExponent),
+				ERDGPassFlags::Compute,
+				ComputeShader,
+				PassParameters,
+				PassParameters->IndirectDispatchArgsBuffer, 0);
+		}
+
+	}
+
 	void SolveWeightedLSR(
 		FRDGBuilder& GraphBuilder,
 		const FSceneView& View,
@@ -1502,7 +1857,7 @@ namespace NFORDenoise
 		const FRDGTextureRef& Radiance,
 		const FRDGBufferRef& NonLocalMeanWeightsBuffer,
 		const FRDGTextureRef& FilteredRadiance,
-		FWeightedLSRDesc WeightedLSRDesc,
+		const FWeightedLSRDesc& WeightedLSRDesc,
 		const FRDGBufferRef Radiances,
 		const FRDGTextureRef& SourceAlbedo
 	)
@@ -1586,31 +1941,16 @@ namespace NFORDenoise
 		}
 		else
 		{
-			typedef RegressionKernel::FLinearSolverCS SHADER;
-			SHADER::FParameters* PassParameters = GraphBuilder.AllocParameters<SHADER::FParameters>();
-			PassParameters->A = GraphBuilder.CreateSRV(FRDGBufferSRVDesc(AMatrix, PF_R32_FLOAT));
-			PassParameters->ADim = FIntPoint(TotalNumOfFeaturesPerFrame, TotalNumOfFeaturesPerFrame);
-			PassParameters->B = GraphBuilder.CreateSRV(FRDGBufferSRVDesc(BMatrix, PF_R32_FLOAT));
-			PassParameters->BDim = BDim;
-			PassParameters->Result = GraphBuilder.CreateUAV(FRDGBufferUAVDesc(ReconstructionWeights, PF_R32_FLOAT));
-
-			PassParameters->NumOfElements = NumOfElements;
-			PassParameters->NumOfElementsPerRow = NumOfElementsPerRow;
-
-			SHADER::FPermutationDomain ComputeShaderPermutationVector;
-
-			checkf(BDim.X >= 6 && BDim.X <= 8, TEXT("Number of features should be between 6 and 8"));
-			ComputeShaderPermutationVector.Set<SHADER::FDimNumFeature>(BDim.X);
-
-			TShaderMapRef<SHADER> ComputeShader(GetGlobalShaderMap(GMaxRHIFeatureLevel), ComputeShaderPermutationVector);
-
-			FComputeShaderUtils::AddPass(
+			ApplyLinearSolverGPU(
 				GraphBuilder,
-				RDG_EVENT_NAME("NFOR::BatchedLinearSolver(F=%d, C=%d)", BDim.X, BDim.Y),
-				ERDGPassFlags::Compute,
-				ComputeShader,
-				PassParameters,
-				FComputeShaderUtils::GetGroupCount(FIntPoint(WeightedLSRDesc.Width, WeightedLSRDesc.Height), NON_LOCAL_MEAN_THREAD_GROUP_SIZE));
+				AMatrix,
+				BMatrix,
+				BDim,
+				ReconstructionWeights,
+				WeightedLSRDesc,
+				TotalNumOfFeaturesPerFrame,
+				NumOfElements,
+				NumOfElementsPerRow);
 		}
 
 		// 3. Reconstruct 
diff --git a/Engine/Plugins/Experimental/NFORDenoise/Source/NFORDenoise/Private/NFORDenoiseCS.h b/Engine/Plugins/Experimental/NFORDenoise/Source/NFORDenoise/Private/NFORDenoiseCS.h
--- a/Engine/Plugins/Experimental/NFORDenoise/Source/NFORDenoise/Private/NFORDenoiseCS.h
+++ b/Engine/Plugins/Experimental/NFORDenoise/Source/NFORDenoise/Private/NFORDenoiseCS.h
@@ -605,6 +605,69 @@ namespace NFORDenoise
 
 				SHADER_PARAMETER(int, NumOfElements)
 				SHADER_PARAMETER(int, NumOfElementsPerRow)
+
+				SHADER_PARAMETER(float, Lambda)
+				SHADER_PARAMETER_RDG_BUFFER_UAV(RWBuffer<uint>, RWSuccessAndFailIndexBuffer)
+			END_SHADER_PARAMETER_STRUCT()
+
+			enum class ESolverType : uint32
+			{
+				NewtonSchulz,
+				Cholesky,
+				NewtonCholesky,  
+				MAX
+			};
+
+			static void ModifyCompilationEnvironment(const FGlobalShaderPermutationParameters& InParameters, FShaderCompilerEnvironment& OutEnvironment)
+			{
+				FGlobalShader::ModifyCompilationEnvironment(InParameters, OutEnvironment);
+				OutEnvironment.SetDefine(TEXT("THREAD_GROUP_SIZE"), NON_LOCAL_MEAN_THREAD_GROUP_SIZE);
+			}
+
+			class FDimNumFeature : SHADER_PERMUTATION_RANGE_INT("NUM_FEATURE", 6, 3);
+			class FDimSolverType : SHADER_PERMUTATION_ENUM_CLASS("LINEAR_SOLVER_TYPE", ESolverType);
+			class FDimOutputIndices : SHADER_PERMUTATION_BOOL("OUTPUT_INDICES");
+			using FPermutationDomain = TShaderPermutationDomain<FDimNumFeature, FDimSolverType, FDimOutputIndices>;
+		};
+
+		enum class EInputMatrixType : uint32
+		{
+			Success,
+			Fail,
+			MAX
+		};
+
+		class FLinearSolverBuildIndirectDispatchArgsCS : public FGlobalShader
+		{
+			DECLARE_GLOBAL_SHADER(FLinearSolverBuildIndirectDispatchArgsCS);
+			SHADER_USE_PARAMETER_STRUCT(FLinearSolverBuildIndirectDispatchArgsCS, FGlobalShader);
+		public:
+
+			BEGIN_SHADER_PARAMETER_STRUCT(FParameters, )
+				SHADER_PARAMETER_RDG_BUFFER_SRV(Buffer<uint>, SuccessAndFailIndexBuffer)
+				SHADER_PARAMETER_RDG_BUFFER_UAV(RWBuffer<uint>, RWIndirectDispatchArgsBuffer)
+			END_SHADER_PARAMETER_STRUCT()
+
+			static void ModifyCompilationEnvironment(const FGlobalShaderPermutationParameters& InParameters, FShaderCompilerEnvironment& OutEnvironment)
+			{
+				FGlobalShader::ModifyCompilationEnvironment(InParameters, OutEnvironment);
+				OutEnvironment.SetDefine(TEXT("THREAD_GROUP_SIZE"), NON_LOCAL_MEAN_THREAD_GROUP_SIZE);
+			}
+
+			class FDimInputMatrixType : SHADER_PERMUTATION_ENUM_CLASS("INPUT_MATRIX_TYPE", EInputMatrixType);
+			using FPermutationDomain = TShaderPermutationDomain <FDimInputMatrixType>;
+		};
+
+		class FLinearSolverIndirectCS : public FGlobalShader
+		{
+			DECLARE_GLOBAL_SHADER(FLinearSolverIndirectCS);
+			SHADER_USE_PARAMETER_STRUCT(FLinearSolverIndirectCS, FGlobalShader);
+
+		public:
+			BEGIN_SHADER_PARAMETER_STRUCT(FParameters, )
+				SHADER_PARAMETER_STRUCT_INCLUDE(FLinearSolverCS::FParameters, CommonParameters)
+				RDG_BUFFER_ACCESS(IndirectDispatchArgsBuffer, ERHIAccess::IndirectArgs)
+				SHADER_PARAMETER_RDG_BUFFER_SRV(Buffer<uint>, SuccessAndFailIndexBuffer)
 			END_SHADER_PARAMETER_STRUCT()
 
 			static void ModifyCompilationEnvironment(const FGlobalShaderPermutationParameters& InParameters, FShaderCompilerEnvironment& OutEnvironment)
@@ -614,7 +677,10 @@ namespace NFORDenoise
 			}
 
 			class FDimNumFeature : SHADER_PERMUTATION_RANGE_INT("NUM_FEATURE", 6, 3);
-			using FPermutationDomain = TShaderPermutationDomain<FDimNumFeature>;
+			class FDimSolverType : SHADER_PERMUTATION_ENUM_CLASS("LINEAR_SOLVER_TYPE", FLinearSolverCS::ESolverType);
+			class FDimInputMatrixType : SHADER_PERMUTATION_ENUM_CLASS("INPUT_MATRIX_TYPE", EInputMatrixType);
+			class FDimOutputIndices : SHADER_PERMUTATION_BOOL("OUTPUT_INDICES");
+			using FPermutationDomain = TShaderPermutationDomain<FDimNumFeature, FDimSolverType, FDimInputMatrixType, FDimOutputIndices>;
 		};
 
 		// reconstruct with the weights into an image
diff --git a/Engine/Plugins/Experimental/NFORDenoise/Source/NFORDenoise/Shaders/NFORDenoise.usf b/Engine/Plugins/Experimental/NFORDenoise/Source/NFORDenoise/Shaders/NFORDenoise.usf
--- a/Engine/Plugins/Experimental/NFORDenoise/Source/NFORDenoise/Shaders/NFORDenoise.usf
+++ b/Engine/Plugins/Experimental/NFORDenoise/Source/NFORDenoise/Shaders/NFORDenoise.usf
@@ -3,6 +3,7 @@
 #include "/Engine/Public/Platform.ush"
 #include "/Engine/Private/Common.ush"
 #include "/Engine/Private/ScreenPass.ush"
+#include "NFORRegressionCommon.ush"
 
 #ifndef SOURCE_CHANNEL_COUNT
 #define SOURCE_CHANNEL_COUNT 4
@@ -96,6 +97,21 @@
 #define WEIGHTED_MULTIPLICATION_TYPE WEIGHTED_MULTIPLICATION_QUADRATIC
 #endif
 
+// Linear equation system solver
+#define INPUT_MATRIX_TYPE_SUCCESS	0
+#define INPUT_MATRIX_TYPE_FAIL		1
+
+#ifndef INPUT_MATRIX_TYPE
+#define INPUT_MATRIX_TYPE	INPUT_MATRIX_TYPE_SUCCESS
+#endif
+
+#define OUTPUT_MATRIX_TYPE_SUCCESS	0
+#define OUTPUT_MATRIX_TYPE_FAIL		1
+
+#ifndef OUTPUT_INDICES
+#define OUTPUT_INDICES 0
+#endif
+
 // Reconstruction
 #define RECONSTRUCTION_TYPE_SCATTER 0
 #define RECONSTRUCTION_TYPE_GATHER	1
@@ -1082,24 +1098,111 @@ void InPlaceBatchedMatrixMultiplicationCS(in const uint3 DispatchThreadID : SV_D
 #endif
 }
 
+#if LINEAR_SOLVER_TYPE == LINEAR_SOLVER_TYPE_NEWTON_SCHULZ
+
+#define NUM_NEWTON_SCHULTZ_ITERATIONS 20
+
+#define NEWTON_INITIAL_GUESS_TYPE INITIAL_GUESS_EUCLIDEAN_NORM
+
+#elif LINEAR_SOLVER_TYPE == LINEAR_SOLVER_TYPE_NEWTON_CHOLESKY
+
+// lambda = 2.5e-3, newton iteration count = 3, equivilent to 20 newton iterations.
+#define NUM_NEWTON_SCHULTZ_ITERATIONS 3
+
+#define NEWTON_INITIAL_GUESS_TYPE INITIAL_GUESS_INVERSE_CHOLESKY_DECOMPOSITION
+
+#endif
+
 #define MATRIX_DIM NUM_FEATURE
 #include "NFORRegression.ush"
 
+float Lambda;
+RWBuffer<uint> RWSuccessAndFailIndexBuffer;
+
+void LinearSolveEntry(uint AOffset, uint ASize, uint BOffset, uint MatrixIndex)
+{
+
+#if (LINEAR_SOLVER_TYPE == LINEAR_SOLVER_TYPE_NEWTON_SCHULZ || \
+	LINEAR_SOLVER_TYPE == LINEAR_SOLVER_TYPE_NEWTON_CHOLESKY)
+	bool bComplete = NewtonIterativeSolve(AOffset, ASize, BOffset, Lambda, Result);
+#elif LINEAR_SOLVER_TYPE == LINEAR_SOLVER_TYPE_CHOLESKY
+	bool bComplete = CholeskeySolve(AOffset, ASize, BOffset, Lambda, Result);
+#endif
+
+#if OUTPUT_INDICES
+	bool WriteCountIndexLocation = bComplete ? OUTPUT_MATRIX_TYPE_SUCCESS : OUTPUT_MATRIX_TYPE_FAIL;
+
+	uint IndexToStore = 0;
+	InterlockedAdd(RWSuccessAndFailIndexBuffer[WriteCountIndexLocation], 1, IndexToStore);
+	IndexToStore = lerp(IndexToStore, (NumOfElements - IndexToStore - 1), WriteCountIndexLocation);
+	RWSuccessAndFailIndexBuffer[IndexToStore + 2] = MatrixIndex;
+#endif
+
+}
+
 [numthreads(THREAD_GROUP_SIZE, THREAD_GROUP_SIZE, 1)]
 void LinearSolverCS(in const uint3 DispatchThreadID : SV_DispatchThreadID)
 {
 	const int MatrixIndex = DispatchThreadID.x + DispatchThreadID.y * NumOfElementsPerRow;
-	if (MatrixIndex >= NumOfElements)
+	if (MatrixIndex >= NumOfElements || DispatchThreadID.x >= NumOfElementsPerRow)
+	{
+		return;
+	}
+	const int BSize = BDim.x * BDim.y;
+	const int BOffset = MatrixIndex * BSize;
+	const int ASize = ADim.x * ADim.y;
+	const int AOffset = MatrixIndex * ASize;
+
+	LinearSolveEntry(AOffset, ASize, BOffset, MatrixIndex);
+}
+
+
+RWBuffer<uint> RWIndirectDispatchArgsBuffer;
+Buffer<uint> SuccessAndFailIndexBuffer;
+
+[numthreads(1, 1, 1)]
+void LinearSolverBuildIndirectDispatchArgsCS(in const uint3 DispatchThreadID : SV_DispatchThreadID)
+{
+	const uint NumOfIndices = SuccessAndFailIndexBuffer[INPUT_MATRIX_TYPE];
+
+	if (NumOfIndices > 0)
+	{
+		const uint NumOfGroupsX = (NumOfIndices + (THREAD_GROUP_SIZE * THREAD_GROUP_SIZE) - 1) / (THREAD_GROUP_SIZE * THREAD_GROUP_SIZE);
+		WriteDispatchIndirectArgs(RWIndirectDispatchArgsBuffer, 0, NumOfGroupsX, 1, 1);
+	}
+	else
+	{
+		WriteDispatchIndirectArgs(RWIndirectDispatchArgsBuffer, 0, 0, 0, 0);
+	}
+}
+
+groupshared uint CompactNumOfElements;
+
+[numthreads(THREAD_GROUP_SIZE*THREAD_GROUP_SIZE, 1, 1)]
+void LinearSolverIndirectCS(in const uint3 DispatchThreadID : SV_DispatchThreadID, uint GI : SV_GroupIndex)
+{
+	if (GI == 0)
+	{
+		CompactNumOfElements = SuccessAndFailIndexBuffer[INPUT_MATRIX_TYPE];
+	}
+
+	GroupMemoryBarrierWithGroupSync();
+
+	const uint ElementIndex = DispatchThreadID.x;
+	if (ElementIndex >= CompactNumOfElements)
 	{
 		return;
 	}
+
+	const uint ResolvedElementIndex = lerp(ElementIndex, NumOfElements - 1 - ElementIndex, INPUT_MATRIX_TYPE);
+	const uint MatrixIndex = SuccessAndFailIndexBuffer[2 + ResolvedElementIndex];
+
 	const int BSize = BDim.x * BDim.y;
 	const int BOffset = MatrixIndex * BSize;
 	const int ASize = ADim.x * ADim.y;
 	const int AOffset = MatrixIndex * ASize;
 
-	//QR_Solve(AOffset, ASize, BOffset);
-	NewtonSchulzIterationDouble(AOffset, ASize, BOffset);
+	LinearSolveEntry(AOffset, ASize, BOffset, MatrixIndex);
 }
 
 //===============================================================================================
diff --git a/Engine/Plugins/Experimental/NFORDenoise/Source/NFORDenoise/Shaders/NFORRegression.ush b/Engine/Plugins/Experimental/NFORDenoise/Source/NFORDenoise/Shaders/NFORRegression.ush
--- a/Engine/Plugins/Experimental/NFORDenoise/Source/NFORDenoise/Shaders/NFORRegression.ush
+++ b/Engine/Plugins/Experimental/NFORDenoise/Source/NFORDenoise/Shaders/NFORRegression.ush
@@ -12,18 +12,21 @@
 
 #pragma once
 
+#include "NFORRegressionCommon.ush"
+
 #ifndef MATRIX_DIM
 #define MATRIX_DIM 8
 #endif
 
-#define FEATURE_MATRIX_SIZE ((MATRIX_DIM)*(MATRIX_DIM))
+#define FULL_MATRIX_SIZE ((MATRIX_DIM)*(MATRIX_DIM))
+#define SYMMETRIC_MATRIX_SIZE (((MATRIX_DIM+1)*MATRIX_DIM)/2)
+#define LOWER_TRIANGLE_MATRIX_SIZE SYMMETRIC_MATRIX_SIZE
 
 #define QR_DECOMPOSITION_TYPE_DEFAULT		0
 #define QR_DECOMPOSITION_TYPE_COLUMN_PIVOT	1
 
 #define QR_DECOMPOSITION_TYPE QR_DECOMPOSITION_TYPE_DEFAULT
 
-
 Buffer<float> A;
 Buffer<float> B;
 
@@ -32,32 +35,142 @@ int2 BDim;
 int NumOfElements;
 int NumOfElementsPerRow;
 
+struct FLMatrix
+{
+	FDFScalar M[LOWER_TRIANGLE_MATRIX_SIZE];
+};
+
+#define LMatrixIndex2D(a,b) select(a>b, int2(b,a), int2(a,b))
+#define LMatrixIndex1D(i) LMatrixIndex2D((i)%MATRIX_DIM, (i)/MATRIX_DIM)
+#define LMatrix(_L,y,x) (_L.M[(x) + (((y) +1) * (y)/2)])
+
+// L^T(x,y) = L(y,x), upper matrix has x>y
+#define TLMatrixIndex2D(a,b) select(a>b, int2(a,b), int2(b,a))
+#define TLMatrixIndex1D(i) TLMatrixIndex2D((i)%MATRIX_DIM, (i)/MATRIX_DIM)
+#define TLMatrix(_L,y,x) LMatrix(_L,x,y)
+
 struct FNFORVector
 {
 	FDFScalar V[MATRIX_DIM];
 };
 
-struct FNFORMatrix
+struct FMatrix
 {
-	FDFScalar M[FEATURE_MATRIX_SIZE];
+	FDFScalar M[FULL_MATRIX_SIZE];
 };
 
-void Identity(out FNFORMatrix Matrix)
+#define MATRIX(_A,y,x) (_A.M[y * MATRIX_DIM + x]) 
+
+void LMatrixNegate(inout FLMatrix A)
+{
+	for (int y = 0; y < MATRIX_DIM; ++y)
+	{
+		for (int x = 0; x <= y; ++x)
+		{
+			LMatrix(A, y, x) = DFNegate(LMatrix(A, y, x));
+		}
+	}
+}
+
+// Assume A and B is symmetric and stored in FLMatrix.
+// A@B != B@A even if A=A^T, B^T=B
+FMatrix  LMatrixMultiply(in FLMatrix A, in FLMatrix B)
+{
+	FMatrix Ret = (FMatrix)0;
+
+	for (int y = 0; y < MATRIX_DIM; ++y)
+	{
+		for (int x = 0; x < MATRIX_DIM; ++x)
+		{
+			FDFScalar Sum = (FDFScalar)0;
+			for (int k = 0; k < MATRIX_DIM; ++k)
+			{
+				int2 yk = LMatrixIndex2D(y,k);
+				int2 kx = LMatrixIndex2D(k,x);
+				Sum = DFAdd(Sum, DFMultiply(LMatrix(A, yk.y, yk.x), LMatrix(B, kx.y, kx.x)));
+			}
+
+			MATRIX(Ret, y, x) = Sum;
+		}
+	}
+
+	return Ret;
+}
+
+// Assume A is symmetric, and B is not. However, A*B is symmetric
+FLMatrix  ToLMatrixMultiply(in FLMatrix A, in FMatrix B)
+{
+	FLMatrix Ret = (FLMatrix)0;
+
+	for (int y = 0; y < MATRIX_DIM; ++y)
+	{
+		for (int x = 0; x <= y; ++x)
+		{
+			FDFScalar Sum = (FDFScalar)0;
+			for (int k = 0; k < MATRIX_DIM; ++k)
+			{
+				int2 yk = LMatrixIndex2D(y,k);
+				Sum = DFAdd(Sum, DFMultiply(LMatrix(A, yk.y, yk.x), MATRIX(B, k, x)));
+			}
+
+			LMatrix(Ret, y, x) = Sum;
+		}
+	}
+
+	return Ret;
+}
+
+void Identity(out FMatrix M)
+{
+	M = (FMatrix)0;
+
+	for (int i = 0; i < MATRIX_DIM; ++i)
+	{
+		MATRIX(M, i, i) = DFPromote(1.0f);
+	}
+}
+
+void LMatrixEuclideanNorm(in FLMatrix A, out float EuclideanNorm)
+{
+	EuclideanNorm = 0.0f;
+
+	for (int i = 0; i < FULL_MATRIX_SIZE; ++i)
+	{
+		int2 Index2D = LMatrixIndex1D(i);
+		EuclideanNorm += Pow2(DFDemote(LMatrix(A, Index2D.y, Index2D.x)));
+	}
+	EuclideanNorm = sqrt(EuclideanNorm);
+}
+
+// |AX-I|. Euclidean norm of AX-I
+void MatrixConvergence(in FMatrix A, out float ConvergenceCriteria)
 {
+	ConvergenceCriteria = 0.0f;
+
 	for (int i = 0; i < MATRIX_DIM; ++i)
 	{
 		for (int j = 0; j < MATRIX_DIM; ++j)
 		{
-			Matrix.M[i * MATRIX_DIM + j] = DFTwoSum(lerp(0.0f, 1.0f, i == j), 0.0f);
+			ConvergenceCriteria += Pow2(DFDemote(MATRIX(A,i,j))-select(i==j, 1.0f, 0.0f));
 		}
 	}
+	ConvergenceCriteria = sqrt(ConvergenceCriteria);
 }
 
-FNFORMatrix Copy(in FNFORMatrix Matrix)
+
+void MatrixNegate(inout FMatrix Matrix)
 {
-	FNFORMatrix Out;
+	for (int i = 0; i < FULL_MATRIX_SIZE; ++i)
+	{
+		Matrix.M[i] = DFNegate(Matrix.M[i]);
+	}
+}
 
-	for (int i = 0; i < FEATURE_MATRIX_SIZE; ++i)
+FMatrix Copy(in FMatrix Matrix)
+{
+	FMatrix Out;
+
+	for (int i = 0; i < FULL_MATRIX_SIZE; ++i)
 	{
 		Out.M[i] = Matrix.M[i];
 	}
@@ -65,9 +178,9 @@ FNFORMatrix Copy(in FNFORMatrix Matrix)
 	return Out;
 }
 
-FNFORMatrix NFORMatrixMultiply(in FNFORMatrix A, in FNFORMatrix B)
+FMatrix MatrixMultiply(in FMatrix A, in FMatrix B)
 {
-	FNFORMatrix Ret;
+	FMatrix Ret;
 	for (int i = 0; i < MATRIX_DIM; ++i)
 	{
 		for (int k = 0; k < MATRIX_DIM; ++k)
@@ -106,7 +219,7 @@ void SwapIntVector(inout int V[MATRIX_DIM], int i, int j)
 	V[j] = Tmp;
 }
 
-void SwampColumn(inout FNFORMatrix Matrix, int i, int j)
+void SwampColumn(inout FMatrix Matrix, int i, int j)
 {
 	FDFScalar Tmp;
 	for (int row = 0; row < MATRIX_DIM; ++row)
@@ -117,9 +230,9 @@ void SwampColumn(inout FNFORMatrix Matrix, int i, int j)
 	}
 }
 
-FNFORMatrix SelectSubMatrix(in FNFORMatrix Matrix, int StartRowColumnIndex)
+FMatrix SelectSubMatrix(in FMatrix Matrix, int StartRowColumnIndex)
 {
-	FNFORMatrix Ret;
+	FMatrix Ret;
 	for (int j = StartRowColumnIndex; j < MATRIX_DIM; ++j)
 	{
 		for (int k = StartRowColumnIndex; k < MATRIX_DIM; ++k)
@@ -130,11 +243,11 @@ FNFORMatrix SelectSubMatrix(in FNFORMatrix Matrix, int StartRowColumnIndex)
 	return Ret;
 }
 
-int GetPivot(in FNFORMatrix Matrix, int N)
+int GetPivot(in FMatrix Matrix, int N)
 {
 	FNFORVector Vector;
 
-	FDFScalar MaxNorm = DFTwoSum(-1.0f, 0.0f);
+	FDFScalar MaxNorm = DFPromote(-1.0f);
 	int MaxIndex = 0;
 
 	for (int column = 0; column < N; ++column)
@@ -154,7 +267,7 @@ int GetPivot(in FNFORMatrix Matrix, int N)
 	return MaxIndex + (MATRIX_DIM - N);
 }
 
-void HouseHoldReflection(in FDFScalar A[MATRIX_DIM], out FNFORMatrix O, int N)
+void HouseHoldReflection(in FDFScalar A[MATRIX_DIM], out FMatrix O, int N)
 {
 	FDFScalar V[MATRIX_DIM];
 
@@ -187,8 +300,7 @@ void HouseHoldReflection(in FDFScalar A[MATRIX_DIM], out FNFORMatrix O, int N)
 
 // Calculate AX=B with (Household) QR decomposition and back substitution.
 // TODO: Reduce the memory used in calculation.
-// TODO: Since A matrix is symmetric, try Choleskey Decomposition.
-void QRDecomposition(in FNFORMatrix A, out FNFORMatrix Q, out FNFORMatrix R
+void QRDecomposition(in FMatrix A, out FMatrix Q, out FMatrix R
 #if QR_DECOMPOSITION_TYPE == QR_DECOMPOSITION_TYPE_COLUMN_PIVOT
 	, out int PivotIndices[MATRIX_DIM]
 #endif
@@ -197,8 +309,8 @@ void QRDecomposition(in FNFORMatrix A, out FNFORMatrix Q, out FNFORMatrix R
 	Identity(Q);
 	R = Copy(A);
 
-	FNFORMatrix H;
-	FNFORMatrix Reflection;
+	FMatrix H;
+	FMatrix Reflection;
 	for (int i = 0; i < MATRIX_DIM; ++i)
 	{
 		// Select column with the largest norm as pivot.
@@ -227,16 +339,16 @@ void QRDecomposition(in FNFORMatrix A, out FNFORMatrix Q, out FNFORMatrix R
 			}
 		}
 
-		Q = NFORMatrixMultiply(Q, H);
-		R = NFORMatrixMultiply(H, R);
+		Q = MatrixMultiply(Q, H);
+		R = MatrixMultiply(H, R);
 	}
 }
 
 void QR_Solve(int AOffset, int ASize, int BOffset)
 {
-	FNFORMatrix NA;
-	FNFORMatrix NQ;
-	FNFORMatrix NR;
+	FMatrix NA;
+	FMatrix NQ;
+	FMatrix NR;
 
 #if QR_DECOMPOSITION_TYPE == QR_DECOMPOSITION_TYPE_COLUMN_PIVOT
 	int PivotIndices[MATRIX_DIM];
@@ -248,9 +360,9 @@ void QR_Solve(int AOffset, int ASize, int BOffset)
 #endif
 
 	// QR decomposition.
-	for (int i = 0; i < FEATURE_MATRIX_SIZE; ++i)
+	for (int i = 0; i < FULL_MATRIX_SIZE; ++i)
 	{
-		NA.M[i] = DFTwoSum(A[AOffset + i], 0.0f);
+		NA.M[i] = DFPromote(A[AOffset + i]);
 	}
 
 	QRDecomposition(NA, NQ, NR
@@ -281,7 +393,7 @@ void QR_Solve(int AOffset, int ASize, int BOffset)
 
 	for (int i = 0; i < (MATRIX_DIM * 3); ++i)
 	{
-		X[i] = DFTwoSum(0.0f, 0.0f);
+		X[i] = DFPromote(0.0f);
 	}
 
 	for (int c = 0; c < 3; ++c)
@@ -323,181 +435,283 @@ void QR_Solve(int AOffset, int ASize, int BOffset)
 
 }
 
-// Implement the standard Newton-Schulz Iterative method (float and double)
-// Initialization
-// X = I/EclideanNorm(A)
-// For n iterations
-//	Y = X ( 2*I - A * X);
-//	X = Y
-
-void NewtonSchulzIteration(int AOffset, int ASize, int BOffset)
+bool CholeskeyDecomposition(inout FLMatrix L, FLMatrix AMatrix, float lambda)
 {
-	float X[MATRIX_DIM * MATRIX_DIM];
-	float XT[MATRIX_DIM * MATRIX_DIM];
-	float TI_AX[MATRIX_DIM * MATRIX_DIM];
-
-#define NUM_INVERSE_ITERATIONS_FLOAT 10
+	bool bComplete = true;
+	L = (FLMatrix)0;
 
-	float EuclideanNorm = 0;
-	for (int i = 0; i < ASize; ++i)
-	{
-		EuclideanNorm += Pow2(A[AOffset + i]);
-	}
-	EuclideanNorm = sqrt(EuclideanNorm);
-
-	// Improve the initial guess with identify matrix scaled by Euclidean norm.
-	for (int i = 0; i < ADim.x; ++i)
+	for (int i = 0; i < MATRIX_DIM; ++i)
 	{
-		for (int j = 0; j < ADim.x; ++j)
+		for (int j = 0; j <= i; ++j)
 		{
+			FDFScalar Sum = (FDFScalar)0;
+			for (int k = 0; k < j; ++k)
+			{
+				Sum = DFAdd(Sum, DFMultiply(LMatrix(L, i, k), LMatrix(L, j, k)));
+			}
+
 			if (i == j)
 			{
-				X[i * ADim.x + j] = 1 / EuclideanNorm;
+				// Check Delta > 0
+				FDFScalar Delta = DFAdd(LMatrix(AMatrix,i,i), DFSubtract(lambda, Sum));
+				bComplete &= DFGreater(Delta, 0) ? true : false;
+				LMatrix(L, i, j) = DFSqrt(Delta);
 			}
 			else
 			{
-				X[i * ADim.x + j] = 0;
+				// Check Ljj != 0
+				FDFScalar Ljj = LMatrix(L, j, j);
+				bComplete &= DFEquals(Ljj, 0) ? false : true;
+				LMatrix(L, i, j) = DFDivide(DFSubtract(LMatrix(AMatrix,i,j), Sum), Ljj);
 			}
 		}
 	}
+	return bComplete;
+}
 
-	// Matrix multiplication of float32 will overflow.
-	UNROLL
-		for (int Step = 0; Step < NUM_INVERSE_ITERATIONS_FLOAT; ++Step)
+FLMatrix CholeskyFactorInverse(FLMatrix L)
+{
+	FMatrix Y;
+	//2. Forward substitution
+	// LL^TX=B, LY = B, where B is identity matrix. 
+	for (int c = 0; c < MATRIX_DIM; ++c)
+	{
+		for (int i = 0; i < MATRIX_DIM; ++i)
 		{
-			// 2 I - A * X
-			for (int i = 0; i < ADim.x; ++i)
-			{
-				for (int k = 0; k < ADim.x; ++k)
-				{
-					float Sum = 0;
-					for (int j = 0; j < ADim.x; ++j)
-					{
-						Sum += A[AOffset + (i * ADim.x + j)] * X[j * ADim.x + k];
-					}
-
-					TI_AX[i * ADim.x + k] = -Sum + lerp(0, 2, i == k);
-				}
-			}
-			// X * TI_AX
-			for (int i = 0; i < ADim.x; ++i)
+			FDFScalar Sum = (FDFScalar)0;
+			for (int j = 0; j < i; ++j)
 			{
-				for (int k = 0; k < ADim.x; ++k)
-				{
-					float Sum = 0;
-					for (int j = 0; j < ADim.x; ++j)
-					{
-						Sum += X[i * ADim.x + j] * TI_AX[j * ADim.x + k];
-					}
-					XT[i * ADim.x + k] = Sum;
-				}
+				Sum = DFAdd(Sum, DFMultiply(MATRIX(Y,j,c), LMatrix(L, i, j)));
 			}
 
-			// Copy back
-			for (int i = 0; i < ASize; ++i)
-			{
-				X[i] = XT[i];
-			}
+			MATRIX(Y,i,c) = DFDivide(DFSubtract(lerp(0.0f, 1.0f, i==c), Sum), LMatrix(L, i, i));
 		}
+	}
 
-	//Calculate approx(A^{-1}) * B.
-	for (int i = 0; i < ADim.x; ++i)
+	FLMatrix X = (FLMatrix)0;
+	//3. Backward substitution.
+	//L^TX = Y
+	for (int c = 0; c < MATRIX_DIM; ++c)
 	{
-		for (int k = 0; k < BDim.y; ++k)
+		for (int i = MATRIX_DIM - 1; i >= c; --i)
 		{
-			float Sum = 0;
-			for (int j = 0; j < ADim.y; ++j)
+			FDFScalar Sum = (FDFScalar)0;
+			for (int j = i + 1; j < MATRIX_DIM; ++j)
 			{
-				Sum += X[i * ADim.x + j] * B[BOffset + j * BDim.y + k];
+				int2 jc = TLMatrixIndex2D(j,c);
+				Sum = DFAdd(Sum, DFMultiply(TLMatrix(X, jc.y, jc.x), TLMatrix(L, i, j)));
 			}
 
-			Result[BOffset + i * BDim.y + k] = X[i * ADim.y + k];
+			TLMatrix(X,c,i) = DFDivide(DFSubtract(MATRIX(Y,i,c), Sum), TLMatrix(L, i, i));
 		}
 	}
+
+	return X;
 }
 
-void NewtonSchulzIterationDouble(int AOffset, int ASize, int BOffset)
+bool CholeskeySolve(int AOffset, int ASize, int BOffset, float lambda, RWBuffer<float> Result)
 {
-	FDFScalar X[MATRIX_DIM * MATRIX_DIM];
-	FDFScalar XT[MATRIX_DIM * MATRIX_DIM];
-	FDFScalar TI_AX[MATRIX_DIM * MATRIX_DIM];
-
-#define NUM_INVERSE_ITERATIONS_DOUBLE 20
+	// Decomposte the matrix to L. As the matrix A might not be positive definite,
+	// We solve for (A + \lambda*I)X=B instead, where \lambda is a very small value.
 
-	float EuclideanNorm = 0;
-	for (int i = 0; i < ASize; ++i)
-	{
-		EuclideanNorm += Pow2(A[AOffset + i]);
-	}
-	EuclideanNorm = sqrt(EuclideanNorm);
+	FLMatrix L = (FLMatrix)0;
+	bool bComplete = true;
 
-	// Improve the initial guess with identify matrix scaled by Euclidean norm.
-	for (int i = 0; i < ADim.x; ++i)
+	for (int i = 0; i < MATRIX_DIM; ++i)
 	{
-		for (int j = 0; j < ADim.x; ++j)
+		for (int j = 0; j <= i; ++j)
 		{
+			FDFScalar Sum = (FDFScalar)0;
+			for (int k = 0; k < j; ++k)
+			{
+				Sum = DFAdd(Sum, DFMultiply(LMatrix(L, i, k),LMatrix(L, j, k)));
+			}
+
 			if (i == j)
 			{
-				X[i * ADim.x + j] = DFTwoSum(1 / EuclideanNorm, 0.0f);
+				// Check Delta > 0
+				FDFScalar Delta = DFAdd(A[AOffset + i * MATRIX_DIM + i], DFSubtract(lambda, Sum));
+				bComplete &= DFGreater(Delta, 0) ? true : false;
+				LMatrix(L, i, j) = DFSqrt(Delta);
 			}
 			else
 			{
-				X[i * ADim.x + j] = (FDFScalar)0;
+				// Check Ljj != 0
+				FDFScalar Ljj = LMatrix(L, j, j);
+				bComplete &= DFEquals(Ljj, 0) ? false : true;
+				LMatrix(L, i, j) = DFDivide(DFSubtract(A[AOffset + i * MATRIX_DIM + j], Sum), Ljj);
 			}
 		}
 	}
 
-	for (int Step = 0; Step < NUM_INVERSE_ITERATIONS_DOUBLE; ++Step)
+	FDFScalar Y[MATRIX_DIM * 3];
+	//2. Forward substitution
+	// LL^TX=B, LY = B, 
+	for (int c = 0; c < 3; ++c)
 	{
-		// 2 I - A * X
-		for (int i = 0; i < ADim.x; ++i)
+		for (int i = 0; i < MATRIX_DIM; ++i)
 		{
-			for (int k = 0; k < ADim.x; ++k)
+			FDFScalar Sum = (FDFScalar)0;
+			for (int j = 0; j < i; ++j)
 			{
-				FDFScalar Sum = (FDFScalar)0;
-				for (int j = 0; j < ADim.x; ++j)
-				{
-					const FDFScalar DFMultiplyScalar = DFMultiply(A[AOffset + (i * ADim.x + j)], X[j * ADim.x + k]);
-					Sum = DFAdd(Sum, DFMultiplyScalar);
-				}
-
-				TI_AX[i * ADim.x + k] = DFSubtract(lerp(0, 2, i == k), Sum);
+				Sum = DFAdd(Sum, DFMultiply(Y[j * 3 + c], LMatrix(L, i, j)));
 			}
+
+			Y[i * 3 + c] = DFDivide(DFSubtract(B[BOffset + i * 3 + c], Sum), LMatrix(L, i, i));
 		}
-		// X * TI_AX
-		for (int i = 0; i < ADim.x; ++i)
+	}
+
+	FDFScalar X[MATRIX_DIM * 3];
+	//3. Backward substitution.
+	//L^TX = Y
+	for (int c = 0; c < 3; ++c)
+	{
+		for (int i = MATRIX_DIM-1; i >= 0; --i)
 		{
-			for (int k = 0; k < ADim.x; ++k)
+			FDFScalar Sum = (FDFScalar)0;
+			for (int j = i + 1; j < MATRIX_DIM; ++j)
 			{
-				FDFScalar Sum = (FDFScalar)0;
-				for (int j = 0; j < ADim.x; ++j)
-				{
-					const FDFScalar DFMultiplyScalar = DFMultiply(X[i * ADim.x + j], TI_AX[j * ADim.x + k]);
-					Sum = DFAdd(Sum, DFMultiplyScalar);
-				}
-				XT[i * ADim.x + k] = Sum;
+				Sum = DFAdd(Sum, DFMultiply(X[j * 3 + c], TLMatrix(L, i, j)));
 			}
+
+			X[i * 3 + c] = DFDivide(DFSubtract(Y[i * 3 + c], Sum), TLMatrix(L, i, i));
+		}
+	}
+
+	for (int i = 0; i < (MATRIX_DIM * 3); ++i)
+	{
+		if (isnan(DFDemote(X[i])))
+		{
+			bComplete = false;
+		}
+	}
+
+	if (bComplete)
+	{
+		for (int i = 0; i < (MATRIX_DIM * 3); ++i)
+		{
+			Result[BOffset + i] = DFDemote(X[i]);
+		}
+	}
+
+	return bComplete;
+}
+
+bool InitializeGuess(inout FLMatrix L, FLMatrix AMatrix, float inLambda)
+{
+	bool bComplete = true;
+#if NEWTON_INITIAL_GUESS_TYPE == INITIAL_GUESS_INVERSE_CHOLESKY_DECOMPOSITION
+	bComplete = CholeskeyDecomposition(L, AMatrix, inLambda);
+
+	if (bComplete)
+	{
+		L = CholeskyFactorInverse(L);
+	}
+
+#elif NEWTON_INITIAL_GUESS_TYPE == INITIAL_GUESS_EUCLIDEAN_NORM
+	float EuclideanNorm = 0;
+	LMatrixEuclideanNorm(AMatrix, EuclideanNorm);
+
+	// Improve the initial guess with identify matrix scaled by Euclidean norm.
+	L = (FLMatrix)0;
+
+	for (int x = 0; x < MATRIX_DIM; ++x)
+	{
+		LMatrix(L, x, x) = DFPromote(1 / EuclideanNorm);
+	}
+#else
+#error not implemented
+#endif
+
+	return bComplete;
+}
+
+FLMatrix MatrixInverse(FLMatrix AMatrix, float inLambda, inout bool bComplete)
+{
+	FLMatrix X;
+
+	// Note that X is invertable in the initialization.
+	bComplete = InitializeGuess(X, AMatrix, inLambda);
+
+
+#define NUM_INVERSE_ITERATIONS_DOUBLE NUM_NEWTON_SCHULTZ_ITERATIONS
+
+	FMatrix AX = (FMatrix)0;
+
+#if LINEAR_SOLVER_TYPE == LINEAR_SOLVER_TYPE_NEWTON_CHOLESKY
+	float MaxTolerance = 1000;
+	float CurrentTolerance = 0.0f;
+#endif
+
+	for (int Step = 0; Step < NUM_INVERSE_ITERATIONS_DOUBLE; ++Step)
+	{
+		// A * X is not guranteed to be symmetric
+		AX = LMatrixMultiply(AMatrix, X);
+
+		// The norm might increase when combining choleksy and newton while Cholesky is overfitting.
+		// Should stop the application if norm increase is detected.
+		#if LINEAR_SOLVER_TYPE == LINEAR_SOLVER_TYPE_NEWTON_CHOLESKY && NEWTON_SCHULZ_EARLY_STOP
+		MatrixConvergence(AX, CurrentTolerance);
+		if (CurrentTolerance > MaxTolerance)
+		{
+			break;
 		}
+		MaxTolerance = min(CurrentTolerance, MaxTolerance);
+		#endif
 
-		// copy back
-		for (int i = 0; i < ASize; ++i)
+		MatrixNegate(AX);
+
+		// 2*I + (-AX) 
+		for (int x = 0; x < MATRIX_DIM; ++x)
+		{
+			MATRIX(AX, x, x) = DFAdd(2.0f, MATRIX(AX, x, x));
+		}
+
+		// X ( 2*I - AX) is guranteed to be symmetric
+		X = ToLMatrixMultiply(X, AX);
+	}
+
+	return X;
+}
+
+bool NewtonIterativeSolve(int AOffset, int ASize, int BOffset, float inLambda, RWBuffer<float> Result)
+{
+	FLMatrix AMatrix = (FLMatrix)0;
+
+	for (int y = 0; y < MATRIX_DIM; ++y)
+	{
+		for (int x = 0; x <= y; ++x)
 		{
-			X[i] = XT[i];
+			LMatrix(AMatrix, y, x) = DFPromote(A[AOffset + y * MATRIX_DIM + x]);
 		}
 	}
 
+	bool bComplete = false;
+
+	AMatrix = MatrixInverse(AMatrix, inLambda, bComplete);
+
+	if (!bComplete)
+	{
+		return bComplete;
+	}
+
 	//Calculate approx(A^{-1}) * B.
-	for (int i = 0; i < ADim.x; ++i)
+
+	for (int i = 0; i < MATRIX_DIM; ++i)
 	{
 		for (int k = 0; k < BDim.y; ++k)
 		{
-			FDFScalar Sum = DFTwoSum(0, 0);
-			for (int j = 0; j < ADim.y; ++j)
+			FDFScalar Sum = DFPromote(0.0f);
+			for (int j = 0; j < MATRIX_DIM; ++j)
 			{
-				Sum = DFAdd(Sum, DFMultiply(X[i * ADim.x + j], DFTwoSum(B[BOffset + j * BDim.y + k], 0.0f)));
+				int2 ij = TLMatrixIndex2D(i,j);
+				Sum = DFAdd(Sum, DFMultiply(TLMatrix(AMatrix, ij.y, ij.x), DFPromote(B[BOffset + j * BDim.y + k])));
 			}
 
 			Result[BOffset + i * BDim.y + k] = DFDemote(Sum);
 		}
 	}
-}
\ No newline at end of file
+
+	return bComplete;
+}
+
diff --git a/Engine/Plugins/Experimental/NFORDenoise/Source/NFORDenoise/Shaders/NFORRegressionCommon.ush b/Engine/Plugins/Experimental/NFORDenoise/Source/NFORDenoise/Shaders/NFORRegressionCommon.ush
new file mode 100644
--- /dev/null
+++ b/Engine/Plugins/Experimental/NFORDenoise/Source/NFORDenoise/Shaders/NFORRegressionCommon.ush
@@ -0,0 +1,22 @@
+// Copyright Epic Games, Inc. All Rights Reserved.
+
+/*=============================================================================
+	NFORRegressionCommon.ush
+=============================================================================*/
+
+#define LINEAR_SOLVER_TYPE_NEWTON_SCHULZ		0
+#define LINEAR_SOLVER_TYPE_CHOLESKY				1
+#define LINEAR_SOLVER_TYPE_NEWTON_CHOLESKY		2
+
+#ifndef LINEAR_SOLVER_TYPE
+#define LINEAR_SOLVER_TYPE LINEAR_SOLVER_TYPE_NEWTON_SCHULZ
+#endif
+
+#define INITIAL_GUESS_INVERSE_CHOLESKY_DECOMPOSITION	0
+#define INITIAL_GUESS_EUCLIDEAN_NORM					1
+
+#ifndef NEWTON_INITIAL_GUESS_TYPE
+#define NEWTON_INITIAL_GUESS_TYPE INITIAL_GUESS_INVERSE_CHOLESKY_DECOMPOSITION
+#endif
+
+#define NEWTON_SCHULZ_EARLY_STOP	1
\ No newline at end of file
