When source and target image is different, we cannot gather distance between two points at the same time. Fallback to gather difference for all pixels within searching window.

diff --git a/Engine/Plugins/Experimental/NFORDenoise/Source/NFORDenoise/Private/NFORDenoiseCS.cpp b/Engine/Plugins/Experimental/NFORDenoise/Source/NFORDenoise/Private/NFORDenoiseCS.cpp
--- a/Engine/Plugins/Experimental/NFORDenoise/Source/NFORDenoise/Private/NFORDenoiseCS.cpp
+++ b/Engine/Plugins/Experimental/NFORDenoise/Source/NFORDenoise/Private/NFORDenoiseCS.cpp
@@ -1144,8 +1144,11 @@ namespace NFORDenoise
 			}
 
 			SHADER::FPermutationDomain ComputeShaderPermutationVector;
+			{
+				ComputeShaderPermutationVector.Set<SHADER::FDimensionSeparateSourceTarget>(bSeparateSourceTarget);
+			}
 
-			TShaderMapRef<SHADER> ComputeShader(GetGlobalShaderMap(GMaxRHIFeatureLevel));
+			TShaderMapRef<SHADER> ComputeShader(GetGlobalShaderMap(GMaxRHIFeatureLevel), ComputeShaderPermutationVector);
 
 			FComputeShaderUtils::AddPass(
 				GraphBuilder,
diff --git a/Engine/Plugins/Experimental/NFORDenoise/Source/NFORDenoise/Private/NFORDenoiseCS.h b/Engine/Plugins/Experimental/NFORDenoise/Source/NFORDenoise/Private/NFORDenoiseCS.h
--- a/Engine/Plugins/Experimental/NFORDenoise/Source/NFORDenoise/Private/NFORDenoiseCS.h
+++ b/Engine/Plugins/Experimental/NFORDenoise/Source/NFORDenoise/Private/NFORDenoiseCS.h
@@ -436,6 +436,9 @@ namespace NFORDenoise
 			FGlobalShader::ModifyCompilationEnvironment(InParameters, OutEnvironment);
 			OutEnvironment.SetDefine(TEXT("THREAD_GROUP_SIZE"), NON_LOCAL_MEAN_THREAD_GROUP_SIZE);
 		}
+
+		class FDimensionSeparateSourceTarget : SHADER_PERMUTATION_BOOL("NONLOCALMEAN_SEPARATE_SOURCE");
+		using FPermutationDomain = TShaderPermutationDomain<FDimensionSeparateSourceTarget>;
 	};
 
 	//--------------------------------------------------------------------------------------------------------------------
diff --git a/Engine/Plugins/Experimental/NFORDenoise/Source/NFORDenoise/Shaders/NFORDenoise.usf b/Engine/Plugins/Experimental/NFORDenoise/Source/NFORDenoise/Shaders/NFORDenoise.usf
--- a/Engine/Plugins/Experimental/NFORDenoise/Source/NFORDenoise/Shaders/NFORDenoise.usf
+++ b/Engine/Plugins/Experimental/NFORDenoise/Source/NFORDenoise/Shaders/NFORDenoise.usf
@@ -598,15 +598,18 @@ void NonLocalMeanGetSqauredDistanceToAtlasCS(in const uint3 DispatchThreadID : S
 	uint2 Position = DispatchThreadID.xy + SeparableFilteringRegion.xy;
 	FImageContext PContext = GetImageContext(Position, Image, Variance, VarianceScale);
 
-	// Calculate the symmetric squared distance.
-	int BaseOffsetIndex = (DispatchId * DispatchTileCount + DispatchThreadID.z) * NONLOCALMEAN_ATLAS_SYMMETRIC_PAIR_COUNT;
+	// Calculate the squared distance. When source and target is different, we store the 2 consecutive 
+	// distances instead of one pair.
+#define DISTANCE_QUERY_COUNT ((NONLOCALMEAN_SEPARATE_SOURCE+1)*NONLOCALMEAN_ATLAS_SYMMETRIC_PAIR_COUNT)
+
+	int BaseOffsetIndex = (DispatchId * DispatchTileCount + DispatchThreadID.z) * DISTANCE_QUERY_COUNT;
 	int PatchSearchWidth = PatchDistance * 2 + 1;
 
 	TAtlasValueType Distances = (TAtlasValueType)0;
 	FImageContext QContext;
 
 	UNROLL
-	for (int i = 0; i < NONLOCALMEAN_ATLAS_SYMMETRIC_PAIR_COUNT; ++i)
+	for (int i = 0; i < DISTANCE_QUERY_COUNT; ++i)
 	{
 		int OffsetIndex = BaseOffsetIndex + i;
 
@@ -614,9 +617,14 @@ void NonLocalMeanGetSqauredDistanceToAtlasCS(in const uint3 DispatchThreadID : S
 		int2 Q = Position + Offset;
 
 		QContext = GetImageContext(Q, TargetImage, TargetVariance, VarianceScale);
+#if NONLOCALMEAN_SEPARATE_SOURCE
+		float Distance = GetSquaredDistance(PContext, QContext);
+		Distances[i] = Distance;
+#else
 		float2 Distance = GetSymmetricSquaredDistance(PContext, QContext);
 		Distances[i * 2 + 0] = Distance.x;
 		Distances[i * 2 + 1] = Distance.y;
+#endif
 	}
 
 	uint2 AtlasPosition = GetAtlasPosition(DispatchThreadID, DispatchRegionSize.xy);
@@ -710,26 +718,42 @@ void NonLocalMeanReshapeBufferCS(in const uint3 DispatchThreadID : SV_DispatchTh
 	{
 		return;
 	}
-	
+
 	const uint x2pd = DispatchThreadID.x;
 	const uint y2pd = DispatchThreadID.y;
-	const uint w	= DispatchThreadID.z;
+	const uint w = DispatchThreadID.z;
 
 	// Get from input
-	float2 PatchWeights = GetInputFromBuffer(x2pd, y2pd, w);// .x stores (x,y)->P, weights
-															// .y stores P ->(x,y) weights
-	// For a given w index in upper matrix, convert to (i,j), find the mapping to lower coordinate.
+	// For seperate source
+	// both .x and .y stores Source(x,y)->Target(P)
+	// For same image
+	// .x stores (x,y)->P, weights
+	// .y stores P ->(x,y) weights
+	const float2 PatchWeights = GetInputFromBuffer(x2pd, y2pd, w);
 	const uint PatchWidth = 2 * PatchDistance + 1;
-	const int2 WeightCenterPoint = int2(PatchDistance, PatchDistance);
+	int2 CoordinateXY2pd = int2(x2pd, y2pd);
+	int2 CoordinateXY = CoordinateXY2pd - PatchDistance;
+	int2 RegionSize = TargetBufferDim.yz;
+
+#if NONLOCALMEAN_SEPARATE_SOURCE 
+	for (int i = 0; i < 2; ++i)
+	{
+		int WeightIndex = (2 * w + i);
+		if (WeightIndex < TargetBufferDim.x && all(CoordinateXY < RegionSize) && all(CoordinateXY >= 0))
+		{
+			int2 WeightCoordinateP = int2(WeightIndex % PatchWidth, WeightIndex / PatchWidth);
+			WriteWeight(CoordinateXY, WeightCoordinateP, PatchWeights[i]);
+		}
+	}
+#else
+	// For a given w index in upper matrix, convert to (i,j), find the mapping to lower coordinate.
 	int2 WeightCoordinateP = int2 (w % PatchWidth, w / PatchWidth);
+	const int2 WeightCenterPoint = int2(PatchDistance, PatchDistance);
 	int2 DeltaXY = WeightCoordinateP - WeightCenterPoint;
 	int2 WeightCoordinateQ = WeightCenterPoint - DeltaXY;
-	int2 CoordinateXY2pd = int2(x2pd,y2pd);
 	int2 CoordinateP2pd = CoordinateXY2pd + DeltaXY;
-	
+
 	// Save (x,y) -> P weights
-	int2 RegionSize = TargetBufferDim.yz;
-	int2 CoordinateXY = CoordinateXY2pd - PatchDistance;
 	if (all(CoordinateXY < RegionSize) && all(CoordinateXY >= 0))
 	{
 		WriteWeight(CoordinateXY, WeightCoordinateP, PatchWeights.x);
@@ -740,6 +764,7 @@ void NonLocalMeanReshapeBufferCS(in const uint3 DispatchThreadID : SV_DispatchTh
 	{
 		WriteWeight(CoordinateP, WeightCoordinateQ, PatchWeights.y);
 	}
+#endif
 }
 
 //------------------------------------------------------------------------------------------------------
