Add adaptive tile count setup for NFOR denoiser to improve high resolution rendering performance.

* If r.NFOR.NumOfTile is less than 0, The number of tile is determined based on (max of texture size)/(r.NFOR.Tile.Size)
* r.NFOR.Tile.Size sets the size of each tile, default to 192.
#tests Hillside project UHD 4k rendering

diff --git a/Engine/Plugins/Experimental/NFORDenoise/Source/NFORDenoise/Private/NFORDenoiseCS.cpp b/Engine/Plugins/Experimental/NFORDenoise/Source/NFORDenoise/Private/NFORDenoiseCS.cpp
--- a/Engine/Plugins/Experimental/NFORDenoise/Source/NFORDenoise/Private/NFORDenoiseCS.cpp
+++ b/Engine/Plugins/Experimental/NFORDenoise/Source/NFORDenoise/Private/NFORDenoiseCS.cpp
@@ -78,9 +78,17 @@ namespace NFORDenoise
 
 	TAutoConsoleVariable<int32> CVarNFORNumOfTile(
 		TEXT("r.NFOR.NumOfTile"),
-		10,
-		TEXT("<=1: Use a single dispatch. Could run out of memory.\n")
-		TEXT("n: Divide the image into n x n tiles in [1,32].\n"),
+		-1,
+		TEXT("n: Divide the image into n x n tiles in [1,32].\n")
+		TEXT("0<=x<=1: Use a single dispatch. Could run out of memory.\n")
+		TEXT("-1: Automatically determine the number of tiles based on r.NFOR.Tile.Size and the view size.\n"),
+		ECVF_RenderThreadSafe);
+
+	TAutoConsoleVariable<int32> CVarNFORTileSize(
+		TEXT("r.NFOR.Tile.Size"),
+		192,
+		TEXT("The size of the max length of a tile. The default is selected for best performance based on experiment.\n")
+		TEXT("It takes effect only when r.NFOR.NumOfTile is set to -1. Minimal value = 100.\n"),
 		ECVF_RenderThreadSafe);
 
 	TAutoConsoleVariable<int32> CVarNFORTileDebug(
@@ -307,9 +315,19 @@ namespace NFORDenoise
 		return ResolvedSourceFrameIndex;
 	}
 
-	int32 GetNumOfTiles()
+	int32 GetNumOfTiles(FIntPoint TextureSize)
 	{
-		return FMath::Clamp(CVarNFORNumOfTile.GetValueOnRenderThread(), 1, 32);
+		int32 NumOfTile = CVarNFORNumOfTile.GetValueOnRenderThread();
+		if (NumOfTile < 0)
+		{
+			// If the max texture size is 1920, and tile size is 192, the num of tiles is 10x10.
+			// If it is between 1920 and 2111, it remains the same until it becames 2112, the tiles will be 11x11.
+			const int32 MaxTextureSize = TextureSize.GetMax();
+			const int32 TileSize = FMath::Max(CVarNFORTileSize.GetValueOnRenderThread(), 100);
+
+			NumOfTile = MaxTextureSize / TileSize;
+		}
+		return FMath::Clamp(NumOfTile, 1, 32);
 	}
 
 	bool IsTileDebugEnabled()
@@ -2026,11 +2044,11 @@ namespace NFORDenoise
 		const int NumberOfWeightsPerPixel = SearchingPatchSize * SearchingPatchSize;
 
 		// TODO: adaptive tile size for best performance.
-		const int32 NumOfTilesOneSide = GetNumOfTiles();
+		FIntPoint TextureSize = Radiances[0].Data.Image->Desc.Extent;
+		const int32 NumOfTilesOneSide = GetNumOfTiles(TextureSize);
 		FIntPoint NumOfTiles = FIntPoint(NumOfTilesOneSide, NumOfTilesOneSide);
 		const int32 TotalTileCount = NumOfTilesOneSide * NumOfTilesOneSide;
 
-		FIntPoint TextureSize = Radiances[0].Data.Image->Desc.Extent;
 		FIntPoint TileSize = FMath::DivideAndRoundUp(TextureSize, NumOfTiles);
 		FIntPoint PaddingTileOffset = RadianceNonLocalMeanParameters.PatchDistance;
 		FIntPoint PaddedTileSize = TileSize + PaddingTileOffset * 2;
