Improve non-local mean filtering performance for NFOR denoiser.
1. Change all feature and intermediate buffer (e.g., MSE and selection map) filtering from direct compute to Atlas based. Feature filtering is 7 times faster. Reduce rendering time by 21%.
2. Change NLM weight layout from WeightxWxH to WxHxWeight. Reduce total rendering time by additional 27%.

#tests Runs faster on hillside projectx
#rnx

diff --git a/Engine/Plugins/Experimental/NFORDenoise/Source/NFORDenoise/Private/NFORDenoiseCS.cpp b/Engine/Plugins/Experimental/NFORDenoise/Source/NFORDenoise/Private/NFORDenoiseCS.cpp
--- a/Engine/Plugins/Experimental/NFORDenoise/Source/NFORDenoise/Private/NFORDenoiseCS.cpp
+++ b/Engine/Plugins/Experimental/NFORDenoise/Source/NFORDenoise/Private/NFORDenoiseCS.cpp
@@ -91,6 +91,12 @@ namespace NFORDenoise
 		TEXT("It takes effect only when r.NFOR.NumOfTile is set to -1. Minimal value = 100.\n"),
 		ECVF_RenderThreadSafe);
 
+	TAutoConsoleVariable<int32> CVarNFORTileFeatureTileCountDownScale(
+		TEXT("r.NFOR.Tile.FeatureTileCount.DownScale"),
+		2,
+		TEXT("Reduce the number of tiles by this factor for feature filtering."),
+		ECVF_RenderThreadSafe);
+
 	TAutoConsoleVariable<int32> CVarNFORTileDebug(
 		TEXT("r.NFOR.TileDebug"),
 		0,
@@ -178,6 +184,15 @@ namespace NFORDenoise
 		TEXT("	n: At least the max size of the input tile(2k as default). The larger, the less number of dispatch passes.\n"),
 		ECVF_RenderThreadSafe);
 
+	TAutoConsoleVariable<int32> CVarNFORNonLocalMeanWeightLayout(
+		TEXT("r.NFOR.NonLocalMean.WeightLayout"),
+		2,
+		TEXT("The layout of the weight. It affects the performance and how weights are handled at each stage.\n")
+		TEXT(" 0: Do not use weight buffer.\n")
+		TEXT(" 1: NumOfWeightsPerPixel x Width x Height.\n")
+		TEXT(" 2: Width x Height x NumOfWeightsPerPixel.\n"),
+		ECVF_RenderThreadSafe);
+
 	TAutoConsoleVariable<int32> CVarNFORNonLocalMeanFeaturePatchSize(
 		TEXT("r.NFOR.NonLocalMean.Feature.PatchSize"),
 		3,
@@ -330,6 +345,11 @@ namespace NFORDenoise
 		return FMath::Clamp(NumOfTile, 1, 32);
 	}
 
+	int32 GetFeatureTileSizeDownScale()
+	{
+		return FMath::Max(CVarNFORTileFeatureTileCountDownScale.GetValueOnRenderThread(), 1);
+	}
+
 	bool IsTileDebugEnabled()
 	{
 		return CVarNFORTileDebug.GetValueOnRenderThread() > 0;
@@ -479,6 +499,27 @@ namespace NFORDenoise
 		return static_cast<ENonLocalMeanAtlasType>(AtlasType);
 	}
 
+	ENonLocalMeanWeightLayout GetNonLocalMeanWeightLayout()
+	{
+		const uint32 LayoutType = FMath::Clamp(
+			CVarNFORNonLocalMeanWeightLayout.GetValueOnRenderThread(),
+			0,
+			static_cast<uint32>(ENonLocalMeanWeightLayout::MAX)-1
+		);
+		return static_cast<ENonLocalMeanWeightLayout>(LayoutType);
+	}
+
+	const TCHAR* GetNonLocalMeanWeightLayoutName(ENonLocalMeanWeightLayout WeightLayout)
+	{
+		static const TCHAR* const kEventNames[] = {
+				TEXT("Direct"),
+				TEXT("Buffer(WeightxWxH)"),
+				TEXT("Buffer(WxHxWeight)")
+		};
+		static_assert(UE_ARRAY_COUNT(kEventNames) == int32(ENonLocalMeanWeightLayout::MAX), "Fix me");
+		return kEventNames[int32(WeightLayout)];
+	}
+
 	bool IsBandwidthSelectionEnabled()
 	{
 		return CVarNFORBandwidthSelection.GetValueOnRenderThread();
@@ -882,11 +923,15 @@ namespace NFORDenoise
 		const FNFORTextureDesc& Texture,
 		const FNFORTextureDesc& Variance,
 		EVarianceType VarianceType,
-		const FRDGTextureRef& FilteredTexture)
+		const FRDGTextureRef& FilteredTexture,
+		const FNonLocalMeanWeightDesc& WeightBufferDesc)
 	{
 		FGlobalShaderMap* GlobalShaderMap = GetGlobalShaderMap(GMaxRHIFeatureLevel);
 		FIntPoint TextureSize = Texture.Image->Desc.Extent;
 		const FRDGTextureRef VarianceTexture = Variance.Image ? Variance.Image : GSystemTextures.GetBlackDummy(GraphBuilder);
+		bool bUseWeightBuffer = WeightBufferDesc.WeightBuffer && WeightBufferDesc.WeightLayout != ENonLocalMeanWeightLayout::None;
+		ENonLocalMeanWeightLayout WeightLayout = bUseWeightBuffer ? WeightBufferDesc.WeightLayout : ENonLocalMeanWeightLayout::None;
+		FIntRect FilteringRegion = bUseWeightBuffer ? WeightBufferDesc.Region : FIntRect(FIntPoint::ZeroValue, TextureSize);
 
 		typedef FNonLocalMeanFilteringCS SHADER;
 
@@ -899,6 +944,12 @@ namespace NFORDenoise
 			PassParameters->VarianceChannelOffset = Variance.ChannelOffset;
 			PassParameters->DenoisingChannelCount = Texture.ChannelCount;
 			PassParameters->DenoisedImage = GraphBuilder.CreateUAV(FRDGTextureUAVDesc(FilteredTexture));
+			PassParameters->FilteringRegion = FilteringRegion;
+
+			if (bUseWeightBuffer)
+			{
+				PassParameters->NonLocalMeanWeights = GraphBuilder.CreateSRV(FRDGBufferSRVDesc(WeightBufferDesc.WeightBuffer, PF_R32_FLOAT));
+			}
 		}
 
 		SHADER::FPermutationDomain ComputeShaderPermutationVector;
@@ -907,41 +958,24 @@ namespace NFORDenoise
 			ComputeShaderPermutationVector.Set<SHADER::FDimensionUseGuide>(false);
 			ComputeShaderPermutationVector.Set<SHADER::FDimensionImageChannelCount>(Texture.NumOfChannel);
 			ComputeShaderPermutationVector.Set<SHADER::FDimPreAlbedoDivide>(GetPreAlbedoDivideRecoverPhase());
+			ComputeShaderPermutationVector.Set<SHADER::FDimWeightLayout>(WeightLayout);
 		}
 
 		TShaderMapRef<SHADER> ComputeShader(GlobalShaderMap, ComputeShaderPermutationVector);
 
+		FIntPoint FilteringSize = FilteringRegion.Size();
+
 		FComputeShaderUtils::AddPass(
 			GraphBuilder,
-			RDG_EVENT_NAME("NFOR::FeatureFiltering(%s, Dim=%d,%d)", Texture.Image->Name, TextureSize.X, TextureSize.Y),
+			RDG_EVENT_NAME("NFOR::FeatureFiltering(%s, %s Dim=%d,%d)", 
+				Texture.Image->Name, 
+				GetNonLocalMeanWeightLayoutName(WeightLayout),
+				FilteringSize.X, 
+				FilteringSize.Y),
 			ERDGPassFlags::Compute,
 			ComputeShader,
 			PassParameters,
-			FComputeShaderUtils::GetGroupCount(TextureSize, NON_LOCAL_MEAN_THREAD_GROUP_SIZE));
-	}
-
-	void ApplyNonLocalMeanFilterIfRequired(
-		FRDGBuilder& GraphBuilder,
-		const FNonLocalMeanParameters& NonLocalMeanParameters,
-		const FNFORTextureDesc& Texture,
-		const FNFORTextureDesc& Variance,
-		EVarianceType VarianceType,
-		const FRDGTextureRef& FilteredTexture)
-	{
-		if (ShouldApplyFeatureFiltering())
-		{
-			ApplyNonLocalMeanFilter(
-				GraphBuilder,
-				NonLocalMeanParameters,
-				Texture,
-				Variance,
-				VarianceType,
-				FilteredTexture);
-		}
-		else
-		{
-			AddCopyTexturePass(GraphBuilder, Texture.Image, FilteredTexture);
-		}
+			FComputeShaderUtils::GetGroupCount(FilteringSize, NON_LOCAL_MEAN_THREAD_GROUP_SIZE));
 	}
 
 	void GetNLMWeigthsWithAtlas(
@@ -949,8 +983,7 @@ namespace NFORDenoise
 		const FSceneView& View,
 		const FRadianceDesc& SourceRadiance,
 		const FRadianceDesc& TargetRadiance,
-		const FRDGBufferRef& NonLocalMeanWeightsBuffer,
-		FIntRect Region,
+		const FNonLocalMeanWeightDesc& NonLocalMeanWeightDesc,
 		const FNonLocalMeanParameters& NonLocalMeanParameters)
 	{
 		const int32 SearchingPatchSize = (NonLocalMeanParameters.PatchDistance * 2 + 1);
@@ -958,9 +991,9 @@ namespace NFORDenoise
 		const FIntPoint TextureSize = SourceRadiance.Data.Image->Desc.Extent;
 		const bool bSeparateSourceTarget = (SourceRadiance.Data.Image != TargetRadiance.Data.Image);
 
-		const FIntRect SeparableFilteringRegion = Region.Inner(-(NonLocalMeanParameters.PatchDistance + NonLocalMeanParameters.PatchSize));
+		const FIntRect SeparableFilteringRegion = NonLocalMeanWeightDesc.Region.Inner(-(NonLocalMeanParameters.PatchDistance + NonLocalMeanParameters.PatchSize));
 		const FIntPoint SeparableFilteringExtent = SeparableFilteringRegion.Size();
-		const FIntPoint WeightQueryRegionExtent = Region.Inner(-(NonLocalMeanParameters.PatchDistance)).Size();
+		const FIntPoint WeightQueryRegionExtent = NonLocalMeanWeightDesc.Region.Inner(-(NonLocalMeanParameters.PatchDistance)).Size();
 		
 		// Estimate the number of dispatches required using the atlas
 		const FIntPoint NLMWeightAtlasExtent = GetNonLocalMeanAtlasSize(SeparableFilteringExtent);
@@ -977,7 +1010,7 @@ namespace NFORDenoise
 		const int32 NumOfDispatch = FMath::DivideAndRoundUp(HalfOffsetSearchCount, SingleDispatchOffsetSearchCount);
 
 		// Allocate the atlas where each pixel stores two symmetric distance/weights, and the temporal buffer.
-		FRDGTextureDesc NLMWeightAtlasDesc= SourceRadiance.Data.Image->Desc;
+		FRDGTextureDesc NLMWeightAtlasDesc = SourceRadiance.Data.Image->Desc;
 		int32 NonLocalMeanWeightsBytesPerElement = 0;
 		{
 			NLMWeightAtlasDesc.Extent = NLMWeightAtlasExtent;
@@ -1155,15 +1188,16 @@ namespace NFORDenoise
 			{
 				PassParameters->NLMParams = NonLocalMeanParameters;
 				PassParameters->SourceBuffer = GraphBuilder.CreateSRV(FRDGBufferSRVDesc(NonLocalMeanWeights, PF_G32R32F));
-				PassParameters->RWTargetBuffer = GraphBuilder.CreateUAV(NonLocalMeanWeightsBuffer, PF_R32_FLOAT);
+				PassParameters->RWTargetBuffer = GraphBuilder.CreateUAV(NonLocalMeanWeightDesc.WeightBuffer, PF_R32_FLOAT);
 				PassParameters->SourceBufferDim = FIntVector4(NumSymmetricPairsPerPixel, WeightQueryRegionExtent.X, WeightQueryRegionExtent.Y, TotalNumTilesToFill);
-				PassParameters->TargetBufferDim = FIntVector(NumberOfWeightsPerPixel, Region.Size().X, Region.Size().Y);
+				PassParameters->TargetBufferDim = FIntVector(NumberOfWeightsPerPixel, NonLocalMeanWeightDesc.Region.Size().X, NonLocalMeanWeightDesc.Region.Size().Y);
 				PassParameters->HalfOffsetSearchCount = HalfOffsetSearchCount;
 			}
 
 			SHADER::FPermutationDomain ComputeShaderPermutationVector;
 			{
 				ComputeShaderPermutationVector.Set<SHADER::FDimensionSeparateSourceTarget>(bSeparateSourceTarget);
+				ComputeShaderPermutationVector.Set<SHADER::FDimensionTargetWeightLayout>(NonLocalMeanWeightDesc.WeightLayout);
 			}
 
 			TShaderMapRef<SHADER> ComputeShader(GetGlobalShaderMap(GMaxRHIFeatureLevel), ComputeShaderPermutationVector);
@@ -1181,6 +1215,103 @@ namespace NFORDenoise
 		}
 	}
 
+	void ApplyNonLocalMeanFilterIfRequired(
+		FRDGBuilder& GraphBuilder,
+		const FSceneView& View,
+		const FNonLocalMeanParameters& NonLocalMeanParameters,
+		const FNFORTextureDesc& Texture,
+		const FNFORTextureDesc& Variance,
+		EVarianceType VarianceType,
+		const FRDGTextureRef& FilteredTexture,
+		int32 TileDownScale = 1)
+	{
+		RDG_EVENT_SCOPE(GraphBuilder, "NonLocalMeanFiltering (%s)", Texture.Image->Name);
+
+		if (ShouldApplyFeatureFiltering())
+		{
+			const FIntPoint TextureSize = Texture.Image->Desc.Extent;
+
+			ENonLocalMeanWeightLayout WeightLayout = GetNonLocalMeanWeightLayout();
+			FNonLocalMeanWeightDesc NonLocalMeanWeightDesc;
+			NonLocalMeanWeightDesc.WeightLayout = WeightLayout;
+
+			const bool bShouldApplyDirectFiltering = 
+				(WeightLayout == ENonLocalMeanWeightLayout::None)
+				|| !ShouldNonLocalMeanUseAtlas();
+
+			if (bShouldApplyDirectFiltering)
+			{
+				NonLocalMeanWeightDesc.Region = FIntRect(FIntPoint(0), TextureSize);
+
+				ApplyNonLocalMeanFilter(
+					GraphBuilder,
+					NonLocalMeanParameters,
+					Texture,
+					Variance,
+					VarianceType,
+					FilteredTexture,
+					NonLocalMeanWeightDesc);
+			}
+			else
+			{
+				/**
+				For each tile:
+					Query the weights.
+					Filter region with the weights.
+				*/
+				// TODO: refactor the tiling common code.
+
+				const int SearchingPatchSize = (NonLocalMeanParameters.PatchDistance * 2 + 1);
+				const int NumberOfWeightsPerPixel = SearchingPatchSize * SearchingPatchSize;
+
+				const int32 NumOfTilesOneSide = FMath::Max(GetNumOfTiles(TextureSize) / TileDownScale, 1);
+				FIntPoint NumOfTiles = FIntPoint(NumOfTilesOneSide, NumOfTilesOneSide);
+				const int32 TotalTileCount = NumOfTilesOneSide * NumOfTilesOneSide;
+
+				FIntPoint TileSize = FMath::DivideAndRoundUp(TextureSize, NumOfTiles);
+
+				const int32 NonLocalMeanSingleFrameWeightSize = TileSize.X * TileSize.Y * NumberOfWeightsPerPixel;
+				const int32 BytesPerElement = sizeof(float);
+				FRDGBufferDesc NonLocalMeanSingleFrameWeightsBufferDesc =
+					FRDGBufferDesc::CreateBufferDesc(BytesPerElement, NonLocalMeanSingleFrameWeightSize);
+				FRDGBufferRef NonLocalMeanSingleFrameWeightsBuffer = GraphBuilder.CreateBuffer(NonLocalMeanSingleFrameWeightsBufferDesc, TEXT("NFOR.NLMFeatureFilteringWeightsBuffer"));
+
+				FRadianceDesc ImageDesc(Texture, Variance, VarianceType);
+				NonLocalMeanWeightDesc.WeightBuffer = NonLocalMeanSingleFrameWeightsBuffer;
+
+				for (int i = 0; i < TotalTileCount; ++i)
+				{
+					int32 TileIndex = i;
+					FIntPoint TileStartPoint = FIntPoint(TileIndex % NumOfTiles.X, TileIndex / NumOfTiles.X) * TileSize;
+					FIntRect TileRegion = FIntRect(FIntPoint(0), TileSize) + TileStartPoint;
+
+					NonLocalMeanWeightDesc.Region = TileRegion;
+
+					GetNLMWeigthsWithAtlas(
+						GraphBuilder,
+						View,
+						ImageDesc,
+						ImageDesc,
+						NonLocalMeanWeightDesc,
+						NonLocalMeanParameters);
+
+					ApplyNonLocalMeanFilter(
+						GraphBuilder,
+						NonLocalMeanParameters,
+						Texture,
+						Variance,
+						VarianceType,
+						FilteredTexture,
+						NonLocalMeanWeightDesc);
+				}
+			}
+		}
+		else
+		{
+			AddCopyTexturePass(GraphBuilder, Texture.Image, FilteredTexture);
+		}
+	}
+
 	void GetNLMWeights(
 		FRDGBuilder& GraphBuilder,
 		const FSceneView& View,
@@ -1198,10 +1329,15 @@ namespace NFORDenoise
 
 		RDG_EVENT_SCOPE(GraphBuilder, "NonLocalMeanGetWeights");
 
+		FNonLocalMeanWeightDesc NonLocalMeanWeightDesc;
+		NonLocalMeanWeightDesc.Region = Region;
+		NonLocalMeanWeightDesc.WeightBuffer = NonLocalMeanWeightsBuffer;
+		NonLocalMeanWeightDesc.WeightLayout = GetNonLocalMeanWeightLayout();
+
 		// Query the non-local mean weights for the radiance.
 		if (bShouldNonLocalMeanUseAtlas)
 		{
-			GetNLMWeigthsWithAtlas(GraphBuilder, View, SourceRadiance, TargetRadiance, NonLocalMeanWeightsBuffer, Region, NonLocalMeanParameters);
+			GetNLMWeigthsWithAtlas(GraphBuilder, View, SourceRadiance, TargetRadiance, NonLocalMeanWeightDesc, NonLocalMeanParameters);
 		}
 		else
 		{
@@ -1214,7 +1350,7 @@ namespace NFORDenoise
 				PassParameters->CommonParameters.TextureSize = TextureSize;
 				PassParameters->CommonParameters.VarianceChannelOffset = SourceRadiance.Variance.ChannelOffset;
 
-				PassParameters->NonLocalMeanWeights = GraphBuilder.CreateUAV(NonLocalMeanWeightsBuffer, PF_R32_FLOAT);
+				PassParameters->RWNonLocalMeanWeights = GraphBuilder.CreateUAV(NonLocalMeanWeightsBuffer, PF_R32_FLOAT);
 				PassParameters->Region = Region;
 
 				if (bSeparateSourceTarget)
@@ -1231,6 +1367,7 @@ namespace NFORDenoise
 				ComputeShaderPermutationVector.Set<SHADER::FDimensionImageChannelCount>(SourceRadiance.Data.NumOfChannel);
 				ComputeShaderPermutationVector.Set<SHADER::FDimensionSeparateSourceTarget>(bSeparateSourceTarget);
 				ComputeShaderPermutationVector.Set<SHADER::FDimPreAlbedoDivide>(GetPreAlbedoDivideRecoverPhase());
+				ComputeShaderPermutationVector.Set<SHADER::FDimTargetWeightLayout>(NonLocalMeanWeightDesc.WeightLayout);
 			}
 
 			TShaderMapRef<SHADER> ComputeShader(GetGlobalShaderMap(GMaxRHIFeatureLevel), ComputeShaderPermutationVector);
@@ -1429,6 +1566,7 @@ namespace NFORDenoise
 				ComputeShaderPermutationVector.Set<SHADER::FDimOptimizeTargetMatrixMultiplication>(true);
 				ComputeShaderPermutationVector.Set<SHADER::FDimNumFeature>(BufferXDimWithConstant);
 				ComputeShaderPermutationVector.Set<SHADER::FDimUseSamplingStep>(SamplingStep > 1);
+				ComputeShaderPermutationVector.Set<SHADER::FDimensionWeightLayout>(GetNonLocalMeanWeightLayout());
 			}
 
 			TShaderMapRef<SHADER> ComputeShader(GetGlobalShaderMap(GMaxRHIFeatureLevel), ComputeShaderPermutationVector);
@@ -1523,6 +1661,7 @@ namespace NFORDenoise
 				ComputeShaderPermutationVector.Set<SHADER::FDimReconstructionType>(ReconstructionType);
 				ComputeShaderPermutationVector.Set<SHADER::FDimPreAlbedoDivide>(GetPreAlbedoDivideRecoverPhase());
 				ComputeShaderPermutationVector.Set<SHADER::FDimNumFeature>(BDim.X);
+				ComputeShaderPermutationVector.Set<SHADER::FDimensionWeightLayout>(GetNonLocalMeanWeightLayout());
 			}
 
 			TShaderMapRef<SHADER> ComputeShader(GetGlobalShaderMap(GMaxRHIFeatureLevel), ComputeShaderPermutationVector);
@@ -2318,11 +2457,13 @@ namespace NFORDenoise
 
 			ApplyNonLocalMeanFilterIfRequired(
 				GraphBuilder,
+				View,
 				FeatureNonLocalMeanParameters,
 				FeatureDesc.Feature,
 				FeatureDesc.Variance,
 				FeatureDesc.VarianceType,
-				FilterdFeature);
+				FilterdFeature,
+				GetFeatureTileSizeDownScale());
 
 			AddCopyTexturePass(GraphBuilder, FilterdFeature, Feature);
 		}
@@ -2441,6 +2582,7 @@ namespace NFORDenoise
 
 				ApplyNonLocalMeanFilterIfRequired(
 					GraphBuilder,
+					View,
 					MSENonLocalMeanParameters,
 					MSE,
 					SourceRadiance.Variance,
@@ -2491,6 +2633,7 @@ namespace NFORDenoise
 
 				ApplyNonLocalMeanFilterIfRequired(
 					GraphBuilder,
+					View,
 					SelectionMapNonLocalMeanParameters,
 					NFORSelectionMap,
 					SourceRadiance.Variance,
diff --git a/Engine/Plugins/Experimental/NFORDenoise/Source/NFORDenoise/Private/NFORDenoiseCS.h b/Engine/Plugins/Experimental/NFORDenoise/Source/NFORDenoise/Private/NFORDenoiseCS.h
--- a/Engine/Plugins/Experimental/NFORDenoise/Source/NFORDenoise/Private/NFORDenoiseCS.h
+++ b/Engine/Plugins/Experimental/NFORDenoise/Source/NFORDenoise/Private/NFORDenoiseCS.h
@@ -283,6 +283,27 @@ namespace NFORDenoise
 
 	FNonLocalMeanParameters GetNonLocalMeanParameters(int32 PatchSize, int32 PatchDistance, float Bandwidth);
 
+	enum class ENonLocalMeanWeightLayout : uint32
+	{
+		/**Weight buffer is not in use*/
+		None,
+		NumOfWeightsPerPixelxWxH,
+		WxHxNumOfWeightsPerPixel,
+		MAX
+	};
+
+	struct FNonLocalMeanWeightDesc
+	{
+		/** The region this weights is gathered for.*/
+		FIntRect Region;
+
+		/** The gathered weights*/
+		FRDGBufferRef WeightBuffer = nullptr;
+
+		/** The layout of the weight*/
+		ENonLocalMeanWeightLayout WeightLayout = ENonLocalMeanWeightLayout::None;
+	};
+
 	// Output the non-local mean filtered image (feature) based on variance
 	// Input: Image, variance, Guide(optional), NonLocalMean parameters
 	//		 Dim: WxHx?, WxHx?
@@ -302,6 +323,11 @@ namespace NFORDenoise
 			SHADER_PARAMETER(FIntPoint, TextureSize)
 			SHADER_PARAMETER(int32, VarianceChannelOffset)
 			SHADER_PARAMETER(int32, DenoisingChannelCount)
+			SHADER_PARAMETER(FIntRect, FilteringRegion)
+
+			//If using non-local mean weights for filtering acceleration
+			SHADER_PARAMETER_RDG_BUFFER_SRV(Buffer<float>, NonLocalMeanWeights)
+
 			SHADER_PARAMETER_RDG_TEXTURE_UAV(RWTexture2D, DenoisedImage)
 			END_SHADER_PARAMETER_STRUCT()
 
@@ -316,7 +342,8 @@ namespace NFORDenoise
 		class FDimensionUseGuide : SHADER_PERMUTATION_BOOL("USE_GUIDE"); //TODO
 		class FDimensionImageChannelCount : SHADER_PERMUTATION_RANGE_INT("SOURCE_CHANNEL_COUNT", 1, static_cast<int>(EImageChannelCount::MAX));
 		class FDimPreAlbedoDivide : SHADER_PERMUTATION_ENUM_CLASS("PRE_ALBEDO_DIVIDE", EAlbedoDivideRecoverPhase);
-		using FPermutationDomain = TShaderPermutationDomain<FDimensionVarianceType, FDimensionUseGuide, FDimensionImageChannelCount, FDimPreAlbedoDivide>;
+		class FDimWeightLayout : SHADER_PERMUTATION_ENUM_CLASS("NLM_WEIGHTLAYOUT", ENonLocalMeanWeightLayout);
+		using FPermutationDomain = TShaderPermutationDomain<FDimensionVarianceType, FDimensionUseGuide, FDimensionImageChannelCount, FDimPreAlbedoDivide, FDimWeightLayout>;
 	};
 
 
@@ -336,7 +363,7 @@ namespace NFORDenoise
 			SHADER_PARAMETER(FIntRect, Region)
 			SHADER_PARAMETER_RDG_TEXTURE_SRV(Texture2D, TargetImage)
 			SHADER_PARAMETER_RDG_TEXTURE_SRV(Texture2D, TargetVariance)
-			SHADER_PARAMETER_RDG_BUFFER_UAV(RWBuffer<float>, NonLocalMeanWeights)
+			SHADER_PARAMETER_RDG_BUFFER_UAV(RWBuffer<float>, RWNonLocalMeanWeights)
 		END_SHADER_PARAMETER_STRUCT()
 
 		static void ModifyCompilationEnvironment(const FGlobalShaderPermutationParameters& InParameters, FShaderCompilerEnvironment& OutEnvironment)
@@ -350,7 +377,9 @@ namespace NFORDenoise
 		class FDimensionImageChannelCount : SHADER_PERMUTATION_RANGE_INT("SOURCE_CHANNEL_COUNT", 1, static_cast<int>(EImageChannelCount::MAX));
 		class FDimensionSeparateSourceTarget : SHADER_PERMUTATION_BOOL("NONLOCALMEAN_SEPARATE_SOURCE");
 		class FDimPreAlbedoDivide : SHADER_PERMUTATION_ENUM_CLASS("PRE_ALBEDO_DIVIDE", EAlbedoDivideRecoverPhase);
-		using FPermutationDomain = TShaderPermutationDomain<FDimensionVarianceType, FDimensionUseGuide, FDimensionImageChannelCount, FDimensionSeparateSourceTarget, FDimPreAlbedoDivide>;
+		class FDimTargetWeightLayout : SHADER_PERMUTATION_ENUM_CLASS("NLM_WEIGHTLAYOUT", ENonLocalMeanWeightLayout);
+		using FPermutationDomain = TShaderPermutationDomain<FDimensionVarianceType, FDimensionUseGuide, FDimensionImageChannelCount, FDimensionSeparateSourceTarget, 
+			FDimPreAlbedoDivide, FDimTargetWeightLayout>;
 	};
 
 	// Optimize the weights query.
@@ -438,7 +467,8 @@ namespace NFORDenoise
 		}
 
 		class FDimensionSeparateSourceTarget : SHADER_PERMUTATION_BOOL("NONLOCALMEAN_SEPARATE_SOURCE");
-		using FPermutationDomain = TShaderPermutationDomain<FDimensionSeparateSourceTarget>;
+		class FDimensionTargetWeightLayout : SHADER_PERMUTATION_ENUM_CLASS("NLM_WEIGHTLAYOUT", ENonLocalMeanWeightLayout);
+		using FPermutationDomain = TShaderPermutationDomain<FDimensionSeparateSourceTarget, FDimensionTargetWeightLayout>;
 	};
 
 	//--------------------------------------------------------------------------------------------------------------------
@@ -588,7 +618,9 @@ namespace NFORDenoise
 			class FDimNumFeature : SHADER_PERMUTATION_RANGE_INT("NUM_FEATURE", 6, 3);
 			class FDimOptimizeTargetMatrixMultiplication : SHADER_PERMUTATION_BOOL("SMALL_MATRIX_OPTIMIZE");
 			class FDimUseSamplingStep : SHADER_PERMUTATION_BOOL("USE_SAMPLING_STEP");
-			using FPermutationDomain = TShaderPermutationDomain<FDimWeightedMultiplicationType, FDimAddConstantFeatureDim, FDimOptimizeTargetMatrixMultiplication, FDimNumFeature, FDimUseSamplingStep>;
+			class FDimensionWeightLayout : SHADER_PERMUTATION_ENUM_CLASS("NLM_WEIGHTLAYOUT", ENonLocalMeanWeightLayout);
+			using FPermutationDomain = TShaderPermutationDomain<FDimWeightedMultiplicationType, FDimAddConstantFeatureDim, FDimOptimizeTargetMatrixMultiplication, 
+				FDimNumFeature, FDimUseSamplingStep, FDimensionWeightLayout>;
 		};
 
 
@@ -750,7 +782,8 @@ namespace NFORDenoise
 			class FDimReconstructionType : SHADER_PERMUTATION_ENUM_CLASS("RECONSTRUCTION_TYPE", EReconstructionType);
 			class FDimPreAlbedoDivide : SHADER_PERMUTATION_ENUM_CLASS("PRE_ALBEDO_DIVIDE", EAlbedoDivideRecoverPhase);
 			class FDimNumFeature : SHADER_PERMUTATION_RANGE_INT("NUM_FEATURE", 6, 3);
-			using FPermutationDomain = TShaderPermutationDomain<FDimReconstructionType, FDimPreAlbedoDivide, FDimNumFeature>;
+			class FDimensionWeightLayout : SHADER_PERMUTATION_ENUM_CLASS("NLM_WEIGHTLAYOUT", ENonLocalMeanWeightLayout);
+			using FPermutationDomain = TShaderPermutationDomain<FDimReconstructionType, FDimPreAlbedoDivide, FDimNumFeature, FDimensionWeightLayout>;
 		};
 	}
 
diff --git a/Engine/Plugins/Experimental/NFORDenoise/Source/NFORDenoise/Shaders/NFORDenoise.usf b/Engine/Plugins/Experimental/NFORDenoise/Source/NFORDenoise/Shaders/NFORDenoise.usf
--- a/Engine/Plugins/Experimental/NFORDenoise/Source/NFORDenoise/Shaders/NFORDenoise.usf
+++ b/Engine/Plugins/Experimental/NFORDenoise/Source/NFORDenoise/Shaders/NFORDenoise.usf
@@ -71,6 +71,14 @@
 
 #define NONLOCALMEAN_ATLAS_SYMMETRIC_PAIR_COUNT (NONLOCALMEAN_ATLAS_TYPE+1)
 
+#define NLM_WEIGHTLAYOUT_NONE						  0
+#define NLM_WEIGHTLAYOUT_NUM_OF_WEIGHTS_PER_PIXELxWxH 1
+#define NLM_WEIGHTLAYOUT_WxHxNUM_OF_WEIGHTS_PER_PIXEL 2
+
+#ifndef NLM_WEIGHTLAYOUT
+#define NLM_WEIGHTLAYOUT NLM_WEIGHTLAYOUT_NONE
+#endif
+
 #ifndef NUM_FEATURE
 #define NUM_FEATURE 0
 #endif
@@ -507,18 +515,37 @@ float GetNonLocalMeanWeight(int2 P, int2 Q)
 }
 
 RWTexture2D<TPixelValue> DenoisedImage;
-RWBuffer<float> NonLocalMeanWeights;
+Buffer<float> NonLocalMeanWeights;
 int DenoisingChannelCount;
+int4 FilteringRegion;
+
+uint GetWeightBufferIndex(uint2 LocalPositionInRegion, uint WeightIndex, uint2 RegionSize)
+{
+#if (NLM_WEIGHTLAYOUT == NLM_WEIGHTLAYOUT_NONE) \
+	|| (NLM_WEIGHTLAYOUT == NLM_WEIGHTLAYOUT_NUM_OF_WEIGHTS_PER_PIXELxWxH)
+	const uint PatchWidth = PatchDistance * 2 + 1;
+	const uint NumOfWeightsPerPixel = Pow2(PatchWidth);
+	uint BufferIndex = WeightIndex + NumOfWeightsPerPixel * (LocalPositionInRegion.x + RegionSize.x * LocalPositionInRegion.y);
+#elif NLM_WEIGHTLAYOUT == NLM_WEIGHTLAYOUT_WxHxNUM_OF_WEIGHTS_PER_PIXEL
+	// Best if neighbour pixels query the same weight index.
+	uint BufferIndex = LocalPositionInRegion.x + RegionSize.x * (LocalPositionInRegion.y + RegionSize.y * WeightIndex);
+#else
+	#error not implemented
+#endif
+	return BufferIndex;
+}
 
 [numthreads(THREAD_GROUP_SIZE, THREAD_GROUP_SIZE, 1)]
 void NonLocalMeanFilteringCS(in const uint3 DispatchThreadID : SV_DispatchThreadID)
 {
-	if (any(DispatchThreadID.xy >= TextureSize))
+	int2 RegionSize = int2(FilteringRegion.z - FilteringRegion.x, FilteringRegion.w - FilteringRegion.y);
+	if (any(DispatchThreadID.xy >= RegionSize))
 	{
 		return;
 	}
-	
-	const uint2 Position = DispatchThreadID.xy;
+
+	const uint PatchWidth = PatchDistance * 2 + 1;
+	const uint2 Position = DispatchThreadID.xy + FilteringRegion.xy;
 	float TotalWeights = 0;
 	TPixelValue Result = 0;
 
@@ -527,8 +554,14 @@ void NonLocalMeanFilteringCS(in const uint3 DispatchThreadID : SV_DispatchThread
 		for (int x = -PatchDistance; x <= PatchDistance; ++x)
 		{
 			int2 Q = Position + int2(x, y);
+#if NLM_WEIGHTLAYOUT == NLM_WEIGHTLAYOUT_NONE
 			float Weight = GetNonLocalMeanWeight(Position, Q);
-
+#else
+			uint WeightIndex = (x + PatchDistance) + (y + PatchDistance) * PatchWidth;
+			uint BufferIndex = GetWeightBufferIndex(DispatchThreadID.xy, WeightIndex, RegionSize);
+			float Weight = NonLocalMeanWeights[BufferIndex];
+#endif
+			//TODO: Improve caching of weights or image value in LDS.
 			Result += Weight * GetImageValue(GetMirroredPosition(Q, TextureSize), TargetImage);
 			TotalWeights += Weight;
 		}
@@ -537,19 +570,20 @@ void NonLocalMeanFilteringCS(in const uint3 DispatchThreadID : SV_DispatchThread
 	Result = Result / max(1e-6f, TotalWeights);
 
 #if SOURCE_CHANNEL_COUNT == 1
-	DenoisedImage[DispatchThreadID.xy] = Result;
+	DenoisedImage[Position] = Result;
 #else
 	TPixelValue TargetValue = GetImageValue(GetMirroredPosition(Position, TextureSize), TargetImage);
 	for (int i = 0; i < DenoisingChannelCount; ++i)
 	{
 		TargetValue[i] = Result[i];
 	}
-	DenoisedImage[DispatchThreadID.xy] = TargetValue;
+	DenoisedImage[Position] = TargetValue;
 #endif
 }
 
 // Calculate the non-local mean weights for Region
 int4 Region;
+RWBuffer<float> RWNonLocalMeanWeights;
 
 [numthreads(THREAD_GROUP_SIZE, THREAD_GROUP_SIZE, 1)]
 void NonLocalMeanWeightsCS(in const uint3 DispatchThreadID : SV_DispatchThreadID)
@@ -568,7 +602,12 @@ void NonLocalMeanWeightsCS(in const uint3 DispatchThreadID : SV_DispatchThreadID
 	{
 		int2 Offset = int2(i % PatchDiameter, i / PatchDiameter) - int2(PatchDistance, PatchDistance);
 		int2 Q = Position + Offset;
-		NonLocalMeanWeights[BufferOffset + i] = GetNonLocalMeanWeight(Position, Q);
+#if NLM_WEIGHTLAYOUT != NLM_WEIGHTLAYOUT_WxHxNUM_OF_WEIGHTS_PER_PIXEL
+		uint BufferIndex = BufferOffset + i;
+#else
+		uint BufferIndex = GetWeightBufferIndex(DispatchThreadID.xy, i, RegionSize);
+#endif
+		RWNonLocalMeanWeights[BufferIndex] = GetNonLocalMeanWeight(Position, Q);
 	}
 }
 
@@ -705,8 +744,13 @@ float2 GetInputFromBuffer(uint x, uint y, uint w)
 void WriteWeight(uint2 PixelPosition, uint2 WeightPosition, float Weight)
 {
 	const uint PatchDiemeter = 2 * PatchDistance + 1;
-	uint w = WeightPosition.x + WeightPosition.y * PatchDiemeter;
-	uint OutputIndex = w + PixelPosition.x * TargetBufferDim.x + PixelPosition.y * TargetBufferDim.x * TargetBufferDim.y;
+	uint WeightIndex = WeightPosition.x + WeightPosition.y * PatchDiemeter;
+
+#if NLM_WEIGHTLAYOUT != NLM_WEIGHTLAYOUT_WxHxNUM_OF_WEIGHTS_PER_PIXEL
+	uint OutputIndex = WeightIndex + PixelPosition.x * TargetBufferDim.x + PixelPosition.y * TargetBufferDim.x * TargetBufferDim.y;
+#else
+	uint OutputIndex = PixelPosition.x + TargetBufferDim.y * (PixelPosition.y + TargetBufferDim.z * WeightIndex);
+#endif
 	RWTargetBuffer[OutputIndex] = Weight;
 }
 
@@ -853,7 +897,7 @@ int2 YDim;
 // |F0 F1|F0 F1|F0 F1|
 //   T0     T1   T2
 // 
-// W stores one weight per pixel, when a frame stored, the next frame is stored
+// W stores one weight per pixel, when a frame stored, the next frame is stored 
 // P1 ... Pn | P1 ... Pn|
 // T0          T1        ...
 
@@ -888,7 +932,13 @@ float GetY(int2 Position, int f)
 float GetW(int2 Position, int WeightIndex, int FrameIndex)
 {
 	const int NumberOfPixels = TextureSize.x * TextureSize.y;
+
+#if NLM_WEIGHTLAYOUT != NLM_WEIGHTLAYOUT_WxHxNUM_OF_WEIGHTS_PER_PIXEL
 	const int Index = NumOfWeigthsPerPixelPerFrame *(NumberOfPixels * FrameIndex + (Position.y * TextureSize.x + Position.x)) + WeightIndex;
+#else
+	const int Index = NumOfWeigthsPerPixelPerFrame * NumberOfPixels * FrameIndex
+					  + Position.x + TextureSize.x * (Position.y + TextureSize.y * WeightIndex);
+#endif
 	return W[Index];
 }
 
