Improve performance of matrix multiplication of X^TWX and X^TWY for NFOR denoiser. Reduce frame time by 14.7%.
1. Cache feature and color into VGPRs.
2. Modify the tile size config from [REDACTED_PHONE]/192=10) to [REDACTED_PHONE]/213=9). Better config to saturate the bandwidth. Default to 9x9 tile instead of 10x10.
#tests Runs faster on hillside project under 1080p
#rnx

diff --git a/Engine/Plugins/Experimental/NFORDenoise/Source/NFORDenoise/Private/NFORDenoiseCS.cpp b/Engine/Plugins/Experimental/NFORDenoise/Source/NFORDenoise/Private/NFORDenoiseCS.cpp
--- a/Engine/Plugins/Experimental/NFORDenoise/Source/NFORDenoise/Private/NFORDenoiseCS.cpp
+++ b/Engine/Plugins/Experimental/NFORDenoise/Source/NFORDenoise/Private/NFORDenoiseCS.cpp
@@ -86,7 +86,7 @@ namespace NFORDenoise
 
 	TAutoConsoleVariable<int32> CVarNFORTileSize(
 		TEXT("r.NFOR.Tile.Size"),
-		192,
+		213,
 		TEXT("The size of the max length of a tile. The default is selected for best performance based on experiment.\n")
 		TEXT("It takes effect only when r.NFOR.NumOfTile is set to -1. Minimal value = 100.\n"),
 		ECVF_RenderThreadSafe);
diff --git a/Engine/Plugins/Experimental/NFORDenoise/Source/NFORDenoise/Shaders/NFORDenoise.usf b/Engine/Plugins/Experimental/NFORDenoise/Source/NFORDenoise/Shaders/NFORDenoise.usf
--- a/Engine/Plugins/Experimental/NFORDenoise/Source/NFORDenoise/Shaders/NFORDenoise.usf
+++ b/Engine/Plugins/Experimental/NFORDenoise/Source/NFORDenoise/Shaders/NFORDenoise.usf
@@ -985,7 +985,7 @@ void SolveMatrixMultiplicationPerMatrixItem(int2 Position, int2 ModifiedXyYyDim,
 
 				float w = GetW(Position, IndexWithinFrame, frame_index); // get w, cache ???, 19*19*T
 
-				int2 TartPosition = Position + Offset;
+				int2 TargetPosition = Position + Offset;
 
 #if APPEND_CONSTANT_DIMENSION_TO_X
 				float x = 1.0f;
@@ -994,17 +994,17 @@ void SolveMatrixMultiplicationPerMatrixItem(int2 Position, int2 ModifiedXyYyDim,
 				BRANCH
 					if (i < XDim.y)
 					{
-						x = GetX(TartPosition, XDim.y * frame_index + i);
+						x = GetX(TargetPosition, XDim.y * frame_index + i);
 					}
 
 				BRANCH
 					if (k < YDim.y)
 					{
-						y = GetY(TartPosition, YDim.y * frame_index + k);
+						y = GetY(TargetPosition, YDim.y * frame_index + k);
 					}
 #else
-				float x = GetX(TartPosition, XDim.y * frame_index + i); // get x from the ith row X^T
-				float y = GetY(TartPosition, YDim.y * frame_index + k); // get y from the kth column of Y
+				float x = GetX(TargetPosition, XDim.y * frame_index + i); // get x from the ith row X^T
+				float y = GetY(TargetPosition, YDim.y * frame_index + k); // get y from the kth column of Y
 #endif
 				R = DFAdd(R, DFMultiply(DFTwoSum(x, 0.0f), DFMultiply(DFTwoSum(sqrt(w), 0.0f), DFTwoSum(y, 0.0f))));
 			}
@@ -1019,15 +1019,51 @@ void SolveMatrixMultiplicationPerMatrixItem(int2 Position, int2 ModifiedXyYyDim,
 	}
 }
 
+#if WEIGHTED_MULTIPLICATION_TYPE == WEIGHTED_MULTIPLICATION_QUADRATIC
+#define MATRIX_CACHE_SIZE (((NUM_FEATURE + APPEND_CONSTANT_DIMENSION_TO_X) * (NUM_FEATURE + APPEND_CONSTANT_DIMENSION_TO_X) + 1) / 2)
+#elif WEIGHTED_MULTIPLICATION_TYPE == WEIGHTED_MULTIPLICATION_GENERALIZED
+#define MATRIX_CACHE_SIZE ((NUM_FEATURE + APPEND_CONSTANT_DIMENSION_TO_X)*3)
+#endif
+
 struct FMatrixCache
+{
+	float Data[MATRIX_CACHE_SIZE];
+};
+
+struct FFeatureXCache
+{
+	float Data[NUM_FEATURE + APPEND_CONSTANT_DIMENSION_TO_X];
+};
+
+struct FFeatureYCache
 {
 #if WEIGHTED_MULTIPLICATION_TYPE == WEIGHTED_MULTIPLICATION_QUADRATIC
-	float Data[((NUM_FEATURE + APPEND_CONSTANT_DIMENSION_TO_X) * (NUM_FEATURE + APPEND_CONSTANT_DIMENSION_TO_X) + 1) / 2];
-#elif WEIGHTED_MULTIPLICATION_TYPE == WEIGHTED_MULTIPLICATION_GENERALIZED
-	float Data[(NUM_FEATURE + APPEND_CONSTANT_DIMENSION_TO_X) * 3];
+	float Data[NUM_FEATURE + APPEND_CONSTANT_DIMENSION_TO_X];
+#else
+	float3 Data;
 #endif
 };
 
+void FillFeatureXCache(out FFeatureXCache FeatureCache,uint2 Position, int2 ModifiedXyYyDim, uint FrameFeatureOffset)
+{
+	for (int i = 0; i < XDim.y; ++i)
+	{
+		FeatureCache.Data[i] = GetX(Position, FrameFeatureOffset + i); 
+	}
+
+#if APPEND_CONSTANT_DIMENSION_TO_X
+	FeatureCache.Data[ModifiedXyYyDim.x - 1] = 1.0f;
+#endif
+}
+
+void FillFeatureYCache(out FFeatureYCache FeatureCache, uint2 Position, int2 ModifiedXyYyDim, uint FrameFeatureOffset)
+{
+	for (int i = 0; i < YDim.y; ++i)
+	{
+		FeatureCache.Data[i] = GetY(Position, FrameFeatureOffset + i);
+	}
+}
+
 int SourceFrameIndex;
 
 void MatrixMultiplicationSmallMatrixOptimization(int2 Position, int2 ModifiedXyYyDim, int PatchWidth, int SingleFrameDataCount)
@@ -1046,7 +1082,7 @@ void MatrixMultiplicationSmallMatrixOptimization(int2 Position, int2 ModifiedXyY
 	{
 		int IndexWithinFrame = step / NumOfFrames;
 		int2 Offset = int2(IndexWithinFrame % PatchWidth, IndexWithinFrame / PatchWidth);
-		int2 TartPosition = Position + Offset;
+		int2 TargetPosition = Position + Offset;
 
 		int t = step % NumOfFrames;
 	#elif SAMPLING_TYPE == SAMPLING_TYPE_FOCUS_CURRENT_FRAME
@@ -1072,14 +1108,14 @@ void MatrixMultiplicationSmallMatrixOptimization(int2 Position, int2 ModifiedXyY
 		}
 
 		int2 Offset = int2(IndexWithinFrame % PatchWidth, IndexWithinFrame / PatchWidth);
-		int2 TartPosition = Position + Offset;
+		int2 TargetPosition = Position + Offset;
 	#endif
 #else
 	for (int j = 0; j < SingleFrameDataCount; ++j)
 	{
 		int IndexWithinFrame = j;
 		int2 Offset = int2(IndexWithinFrame % PatchWidth, IndexWithinFrame / PatchWidth);
-		int2 TartPosition = Position + Offset;
+		int2 TargetPosition = Position + Offset;
 
 		for (int t = 0; t < NumOfFrames; ++t)
 #endif
@@ -1087,6 +1123,13 @@ void MatrixMultiplicationSmallMatrixOptimization(int2 Position, int2 ModifiedXyY
 			int frame_index = t;
 			float sqrtw = sqrt(GetW(Position, IndexWithinFrame, frame_index)); // get w.
 			
+#if WEIGHTED_MULTIPLICATION_TYPE == WEIGHTED_MULTIPLICATION_GENERALIZED
+			FFeatureYCache Y;
+			FillFeatureYCache(Y, TargetPosition, ModifiedXyYyDim, YDim.y * frame_index);
+#endif
+			FFeatureXCache X;
+			FillFeatureXCache(X, TargetPosition, ModifiedXyYyDim, XDim.y * frame_index);
+
 			Index = 0;
 			for (int i = 0; i < ModifiedXyYyDim.x; ++i)
 			{
@@ -1096,25 +1139,12 @@ void MatrixMultiplicationSmallMatrixOptimization(int2 Position, int2 ModifiedXyY
 				for (int k = 0; k < ModifiedXyYyDim.y; ++k)
 #endif
 				{
+					float x = X.Data[i];
 
-#if APPEND_CONSTANT_DIMENSION_TO_X
-					float x = 1.0f;
-					float y = 1.0f;
-
-					BRANCH
-						if (i < XDim.y)
-						{
-							x = GetX(TartPosition, XDim.y * frame_index + i);
-						}
-
-					BRANCH
-						if (k < YDim.y)
-						{
-							y = GetY(TartPosition, YDim.y * frame_index + k);
-						}
-#else
-					float x = GetX(TartPosition, XDim.y * frame_index + i); // get x from the ith row X^T
-					float y = GetY(TartPosition, YDim.y * frame_index + k); // get y from the kth column of Y
+#if WEIGHTED_MULTIPLICATION_TYPE == WEIGHTED_MULTIPLICATION_GENERALIZED
+					float y = Y.Data[k];
+#else 
+					float y = X.Data[k];
 #endif
 					MatrixCache.Data[Index] += x * sqrtw * y;
 					++Index;
