Add new layout 4xWxHx[NumOfWeightsPerPixel/4] (default on) to remove memory reshape, which reduces frame denoising cost by 7% for NFOR denoiser.
#rnx
#tests MRQ rendering of hillside project shows expected results with 5 frame history .

diff --git a/Engine/Plugins/Experimental/NFORDenoise/Source/NFORDenoise/Private/NFORDenoiseCS.cpp b/Engine/Plugins/Experimental/NFORDenoise/Source/NFORDenoise/Private/NFORDenoiseCS.cpp
--- a/Engine/Plugins/Experimental/NFORDenoise/Source/NFORDenoise/Private/NFORDenoiseCS.cpp
+++ b/Engine/Plugins/Experimental/NFORDenoise/Source/NFORDenoise/Private/NFORDenoiseCS.cpp
@@ -186,11 +186,12 @@ namespace NFORDenoise
 
 	TAutoConsoleVariable<int32> CVarNFORNonLocalMeanWeightLayout(
 		TEXT("r.NFOR.NonLocalMean.WeightLayout"),
-		2,
+		3,
 		TEXT("The layout of the weight. It affects the performance and how weights are handled at each stage.\n")
 		TEXT(" 0: Do not use weight buffer.\n")
 		TEXT(" 1: NumOfWeightsPerPixel x Width x Height.\n")
-		TEXT(" 2: Width x Height x NumOfWeightsPerPixel.\n"),
+		TEXT(" 2: Width x Height x NumOfWeightsPerPixel.\n")
+		TEXT(" 3: float4 x Width x Height x DivideAndRoundUp(NumOfWeightsPerPixel,float4)\n"),
 		ECVF_RenderThreadSafe);
 
 	TAutoConsoleVariable<int32> CVarNFORNonLocalMeanFeaturePatchSize(
@@ -470,6 +471,39 @@ namespace NFORDenoise
 		return FMath::Clamp(CVarNFORNonLocalMeanRadiancePatchDistance.GetValueOnRenderThread(), 0, 30);
 	}
 
+	ENonLocalMeanWeightLayout GetNonLocalMeanWeightLayout()
+	{
+		const uint32 LayoutType = FMath::Clamp(
+			CVarNFORNonLocalMeanWeightLayout.GetValueOnRenderThread(),
+			0,
+			static_cast<uint32>(ENonLocalMeanWeightLayout::MAX) - 1
+		);
+		return static_cast<ENonLocalMeanWeightLayout>(LayoutType);
+	}
+
+	int32 GetNonLocalMeanSingleFrameWeightBufferSize(FIntPoint Size, int32 NumOfWeightsPerPixel)
+	{
+		ENonLocalMeanWeightLayout NonLocalMeanWeightLayout = GetNonLocalMeanWeightLayout();
+		int32 NumOfElements = Size.X * Size.Y * NumOfWeightsPerPixel;
+		if (NonLocalMeanWeightLayout == ENonLocalMeanWeightLayout::Float4xWxHxNumOfWeightsPerPixelByFloat4)
+		{
+			NumOfElements = Size.X * Size.Y * FMath::DivideAndRoundUp(NumOfWeightsPerPixel, 4) * 4;
+		}
+		return NumOfElements;
+	}
+
+	const TCHAR* GetNonLocalMeanWeightLayoutName(ENonLocalMeanWeightLayout WeightLayout)
+	{
+		static const TCHAR* const kEventNames[] = {
+				TEXT("Direct"),
+				TEXT("Buffer(WeightxWxH)"),
+				TEXT("Buffer(WxHxWeight)"),
+				TEXT("Buffer(4xWxHx[Weight/4])")
+		};
+		static_assert(UE_ARRAY_COUNT(kEventNames) == int32(ENonLocalMeanWeightLayout::MAX), "Fix me");
+		return kEventNames[int32(WeightLayout)];
+	}
+
 	FIntPoint GetNonLocalMeanAtlasSize(FIntPoint Extent)
 	{
 		int32 AtlasSize = CVarNFORNonLocalMeanAtlasSize.GetValueOnRenderThread();
@@ -496,28 +530,13 @@ namespace NFORDenoise
 			0,
 			static_cast<uint32>(ENonLocalMeanAtlasType::MAX)-1
 			);
-		return static_cast<ENonLocalMeanAtlasType>(AtlasType);
-	}
+		ENonLocalMeanAtlasType NonLocalMeanAtlasType = static_cast<ENonLocalMeanAtlasType>(AtlasType);
 
-	ENonLocalMeanWeightLayout GetNonLocalMeanWeightLayout()
-	{
-		const uint32 LayoutType = FMath::Clamp(
-			CVarNFORNonLocalMeanWeightLayout.GetValueOnRenderThread(),
-			0,
-			static_cast<uint32>(ENonLocalMeanWeightLayout::MAX)-1
-		);
-		return static_cast<ENonLocalMeanWeightLayout>(LayoutType);
-	}
-
-	const TCHAR* GetNonLocalMeanWeightLayoutName(ENonLocalMeanWeightLayout WeightLayout)
-	{
-		static const TCHAR* const kEventNames[] = {
-				TEXT("Direct"),
-				TEXT("Buffer(WeightxWxH)"),
-				TEXT("Buffer(WxHxWeight)")
-		};
-		static_assert(UE_ARRAY_COUNT(kEventNames) == int32(ENonLocalMeanWeightLayout::MAX), "Fix me");
-		return kEventNames[int32(WeightLayout)];
+		if (GetNonLocalMeanWeightLayout() == ENonLocalMeanWeightLayout::Float4xWxHxNumOfWeightsPerPixelByFloat4)
+		{
+			NonLocalMeanAtlasType = ENonLocalMeanAtlasType::TwoSymmetricPair;
+		}
+		return NonLocalMeanAtlasType;
 	}
 
 	bool IsBandwidthSelectionEnabled()
@@ -917,6 +936,17 @@ namespace NFORDenoise
 		return GetNonLocalMeanParameters(PatchSize, PatchDistance, Bandwidth);
 	}
 
+	EPixelFormat GetWeightLayoutPixelFormat(ENonLocalMeanWeightLayout WeightLayout)
+	{
+		EPixelFormat DefaultFormat = PF_R32_FLOAT;
+		if (WeightLayout == ENonLocalMeanWeightLayout::Float4xWxHxNumOfWeightsPerPixelByFloat4)
+		{
+			DefaultFormat = PF_A32B32G32R32F;
+		}
+
+		return DefaultFormat;
+	}
+
 	void ApplyNonLocalMeanFilter(
 		FRDGBuilder& GraphBuilder,
 		const FNonLocalMeanParameters& NonLocalMeanParameters,
@@ -948,7 +978,8 @@ namespace NFORDenoise
 
 			if (bUseWeightBuffer)
 			{
-				PassParameters->NonLocalMeanWeights = GraphBuilder.CreateSRV(FRDGBufferSRVDesc(WeightBufferDesc.WeightBuffer, PF_R32_FLOAT));
+				PassParameters->NonLocalMeanWeights = GraphBuilder.CreateSRV(FRDGBufferSRVDesc(WeightBufferDesc.WeightBuffer, 
+					GetWeightLayoutPixelFormat(WeightLayout)));
 			}
 		}
 
@@ -989,7 +1020,9 @@ namespace NFORDenoise
 		const int32 SearchingPatchSize = (NonLocalMeanParameters.PatchDistance * 2 + 1);
 		const int32 NumberOfWeightsPerPixel = SearchingPatchSize * SearchingPatchSize;
 		const FIntPoint TextureSize = SourceRadiance.Data.Image->Desc.Extent;
-		const bool bSeparateSourceTarget = (SourceRadiance.Data.Image != TargetRadiance.Data.Image);
+		// If we expect float4xWxHx[NumOfWeight/4], we can directly write through.
+		const bool bShouldUseLocalBuffer = NonLocalMeanWeightDesc.WeightLayout != ENonLocalMeanWeightLayout::Float4xWxHxNumOfWeightsPerPixelByFloat4;
+		const bool bSeparateSourceTarget = (SourceRadiance.Data.Image != TargetRadiance.Data.Image)|| !bShouldUseLocalBuffer;
 
 		const FIntRect SeparableFilteringRegion = NonLocalMeanWeightDesc.Region.Inner(-(NonLocalMeanParameters.PatchDistance + NonLocalMeanParameters.PatchSize));
 		const FIntPoint SeparableFilteringExtent = SeparableFilteringRegion.Size();
@@ -1034,10 +1067,17 @@ namespace NFORDenoise
 		};
 		
 		const int32 TotalNumTilesToFill = FMath::DivideAndRoundUp(HalfOffsetSearchCount, NumSymmetricPairsPerPixel);
-		size_t NonLocalMeanWeightsCount = NumSymmetricPairsPerPixel * (WeightQueryRegionExtent.X * WeightQueryRegionExtent.Y) * TotalNumTilesToFill;
-		FRDGBufferDesc NonLocalMeanWeightsDesc = FRDGBufferDesc::CreateBufferDesc(NonLocalMeanWeightsBytesPerElement, NonLocalMeanWeightsCount);
-		FRDGBufferRef  NonLocalMeanWeights = GraphBuilder.CreateBuffer(NonLocalMeanWeightsDesc, TEXT("NFOR.NonLocalMeanWeights"));
-
+		FRDGBufferRef  NonLocalMeanWeights = nullptr;
+		if (bShouldUseLocalBuffer)
+		{
+			size_t NonLocalMeanWeightsCount = NumSymmetricPairsPerPixel * (WeightQueryRegionExtent.X * WeightQueryRegionExtent.Y) * TotalNumTilesToFill;
+			FRDGBufferDesc NonLocalMeanWeightsDesc = FRDGBufferDesc::CreateBufferDesc(NonLocalMeanWeightsBytesPerElement, NonLocalMeanWeightsCount);
+			NonLocalMeanWeights = GraphBuilder.CreateBuffer(NonLocalMeanWeightsDesc, TEXT("NFOR.NonLocalMeanWeights"));
+		}
+		else
+		{
+			NonLocalMeanWeights = NonLocalMeanWeightDesc.WeightBuffer;
+		}
 		// Summery
 		// 1. For each dispatch:
 		//		Calcualte offset for each tile (SeparableFilteringRegion) in the atlas
@@ -1165,6 +1205,7 @@ namespace NFORDenoise
 					ComputeShaderPermutationVector.Set<SHADER::FDimensionSeperablePassType>(PassInfo.SeperablePassType);
 					ComputeShaderPermutationVector.Set<SHADER::FDimAtlasType>(NonLocalMeanAtlasType);
 					ComputeShaderPermutationVector.Set<SHADER::FDimPreAlbedoDivide>(GetPreAlbedoDivideRecoverPhase());
+					ComputeShaderPermutationVector.Set<SHADER::FDimBufferPassThrough>(!bShouldUseLocalBuffer);
 				}
 
 				TShaderMapRef<SHADER> ComputeShader(GetGlobalShaderMap(GMaxRHIFeatureLevel), ComputeShaderPermutationVector);
@@ -1182,6 +1223,7 @@ namespace NFORDenoise
 		} // End dispatch
 
 		// Reshape the buffer from X*Y*Wb to W*X*Y and scatter the results
+		if (bShouldUseLocalBuffer)
 		{
 			typedef FNonLocalMeanReshapeBufferCS SHADER;
 			SHADER::FParameters* PassParameters = GraphBuilder.AllocParameters<SHADER::FParameters>();
@@ -1270,7 +1312,7 @@ namespace NFORDenoise
 
 				FIntPoint TileSize = FMath::DivideAndRoundUp(TextureSize, NumOfTiles);
 
-				const int32 NonLocalMeanSingleFrameWeightSize = TileSize.X * TileSize.Y * NumberOfWeightsPerPixel;
+				const int32 NonLocalMeanSingleFrameWeightSize = GetNonLocalMeanSingleFrameWeightBufferSize(TileSize, NumberOfWeightsPerPixel);
 				const int32 BytesPerElement = sizeof(float);
 				FRDGBufferDesc NonLocalMeanSingleFrameWeightsBufferDesc =
 					FRDGBufferDesc::CreateBufferDesc(BytesPerElement, NonLocalMeanSingleFrameWeightSize);
@@ -1526,6 +1568,7 @@ namespace NFORDenoise
 			const int BufferXDimWithConstant = XDim.Y + (bFeatureAddConstant ? 1 : 0);
 			FIntPoint ResultMatrixDimension = FIntPoint(BufferXDimWithConstant, GeneralizedMultiplication ? YDim.Y : BufferXDimWithConstant);
 			int32 SamplingStep = GetSamplingStep(WDim, ResultMatrixDimension.X * ResultMatrixDimension.Y);
+			ENonLocalMeanWeightLayout NonLocalMeanWeightLayout = GetNonLocalMeanWeightLayout();
 
 			ResultMatrix = SHADER::AllocateResultBuffer(GraphBuilder, TextureSize, ResultMatrixDimension.X, ResultMatrixDimension.Y);
 
@@ -1534,7 +1577,7 @@ namespace NFORDenoise
 				PassParameters->X = GraphBuilder.CreateSRV(FRDGBufferSRVDesc(X, PF_R32_FLOAT));
 				PassParameters->XDim = XDim;
 
-				PassParameters->W = GraphBuilder.CreateSRV(FRDGBufferSRVDesc(W, PF_R32_FLOAT));
+				PassParameters->W = GraphBuilder.CreateSRV(FRDGBufferSRVDesc(W, GetWeightLayoutPixelFormat(NonLocalMeanWeightLayout)));
 				PassParameters->WDim = WDim;
 
 
@@ -1566,7 +1609,7 @@ namespace NFORDenoise
 				ComputeShaderPermutationVector.Set<SHADER::FDimOptimizeTargetMatrixMultiplication>(true);
 				ComputeShaderPermutationVector.Set<SHADER::FDimNumFeature>(BufferXDimWithConstant);
 				ComputeShaderPermutationVector.Set<SHADER::FDimUseSamplingStep>(SamplingStep > 1);
-				ComputeShaderPermutationVector.Set<SHADER::FDimensionWeightLayout>(GetNonLocalMeanWeightLayout());
+				ComputeShaderPermutationVector.Set<SHADER::FDimensionWeightLayout>(NonLocalMeanWeightLayout);
 			}
 
 			TShaderMapRef<SHADER> ComputeShader(GetGlobalShaderMap(GMaxRHIFeatureLevel), ComputeShaderPermutationVector);
@@ -1605,6 +1648,7 @@ namespace NFORDenoise
 		FRDGBufferRef ReconstructionBuffer = nullptr;
 
 		FIntPoint TextureSize = FilteredRadiance->Desc.Extent;
+		ENonLocalMeanWeightLayout NonLocalMeanWeightLayout = GetNonLocalMeanWeightLayout();
 
 		{
 			typedef RegressionKernel::FReconstructSpatialTemporalImage SHADER;
@@ -1638,7 +1682,7 @@ namespace NFORDenoise
 			SHADER::FParameters* PassParameters = GraphBuilder.AllocParameters<SHADER::FParameters>();
 			{
 				PassParameters->X = GraphBuilder.CreateSRV(FRDGBufferSRVDesc(Feature, PF_R32_FLOAT));
-				PassParameters->W = GraphBuilder.CreateSRV(FRDGBufferSRVDesc(NonLocalMeanWeightsBuffer, PF_R32_FLOAT));
+				PassParameters->W = GraphBuilder.CreateSRV(FRDGBufferSRVDesc(NonLocalMeanWeightsBuffer, GetWeightLayoutPixelFormat(NonLocalMeanWeightLayout)));
 				PassParameters->B = GraphBuilder.CreateSRV(FRDGBufferSRVDesc(ReconstructionWeights, PF_R32_FLOAT));
 				PassParameters->RWReconstruction = GraphBuilder.CreateUAV(FilteredRadiance);
 				PassParameters->RWReconstructBuffer = ReconstructionBuffer ? GraphBuilder.CreateUAV(ReconstructionBuffer) : nullptr;
@@ -1661,7 +1705,7 @@ namespace NFORDenoise
 				ComputeShaderPermutationVector.Set<SHADER::FDimReconstructionType>(ReconstructionType);
 				ComputeShaderPermutationVector.Set<SHADER::FDimPreAlbedoDivide>(GetPreAlbedoDivideRecoverPhase());
 				ComputeShaderPermutationVector.Set<SHADER::FDimNumFeature>(BDim.X);
-				ComputeShaderPermutationVector.Set<SHADER::FDimensionWeightLayout>(GetNonLocalMeanWeightLayout());
+				ComputeShaderPermutationVector.Set<SHADER::FDimensionWeightLayout>(NonLocalMeanWeightLayout);
 			}
 
 			TShaderMapRef<SHADER> ComputeShader(GetGlobalShaderMap(GMaxRHIFeatureLevel), ComputeShaderPermutationVector);
@@ -2194,7 +2238,7 @@ namespace NFORDenoise
 		FIntRect  PaddedTileRect = FIntRect(FIntPoint(0,0), PaddedTileSize);
 		const int NumOfCombinedFeatureChannels = GetNumOfCombinedFeatureChannels(Features);
 
-		const int32 NonLocalMeanSingleFrameWeightSize = TileSize.X * TileSize.Y * NumberOfWeightsPerPixel;
+		const int32 NonLocalMeanSingleFrameWeightSize = GetNonLocalMeanSingleFrameWeightBufferSize(TileSize, NumberOfWeightsPerPixel);
 		const int32 BytesPerElement = sizeof(float);
 		FRDGBufferDesc NonLocalMeanSingleFrameWeightsBufferDesc = 
 			FRDGBufferDesc::CreateBufferDesc(BytesPerElement, NonLocalMeanSingleFrameWeightSize);
@@ -2440,7 +2484,6 @@ namespace NFORDenoise
 		FNonLocalMeanParameters FeatureNonLocalMeanParameters = GetFeatureNonLocalMeanParameters(0.5f);
 
 		const int NumOfFeatures = FeatureDescs.Num();
-		typedef FNonLocalMeanFilteringCS SHADER;
 
 		int32 BufferChannelOffset = 0;
 		for (int i = 0; i < NumOfFeatures; ++i)
diff --git a/Engine/Plugins/Experimental/NFORDenoise/Source/NFORDenoise/Private/NFORDenoiseCS.h b/Engine/Plugins/Experimental/NFORDenoise/Source/NFORDenoise/Private/NFORDenoiseCS.h
--- a/Engine/Plugins/Experimental/NFORDenoise/Source/NFORDenoise/Private/NFORDenoiseCS.h
+++ b/Engine/Plugins/Experimental/NFORDenoise/Source/NFORDenoise/Private/NFORDenoiseCS.h
@@ -289,6 +289,7 @@ namespace NFORDenoise
 		None,
 		NumOfWeightsPerPixelxWxH,
 		WxHxNumOfWeightsPerPixel,
+		Float4xWxHxNumOfWeightsPerPixelByFloat4,
 		MAX
 	};
 
@@ -441,7 +442,8 @@ namespace NFORDenoise
 		class FDimensionSeperablePassType : SHADER_PERMUTATION_ENUM_CLASS("NONLOCALMEAN_SEPRERABLE_PASS", ESeperablePassType);
 		class FDimPreAlbedoDivide : SHADER_PERMUTATION_ENUM_CLASS("PRE_ALBEDO_DIVIDE", EAlbedoDivideRecoverPhase);
 		class FDimAtlasType : SHADER_PERMUTATION_ENUM_CLASS("NONLOCALMEAN_ATLAS_TYPE", ENonLocalMeanAtlasType);
-		using FPermutationDomain = TShaderPermutationDomain<FDimensionSeperablePassType, FDimPreAlbedoDivide, FDimAtlasType>;
+		class FDimBufferPassThrough : SHADER_PERMUTATION_BOOL("BUFFER_PASS_THROUGH");
+		using FPermutationDomain = TShaderPermutationDomain<FDimensionSeperablePassType, FDimPreAlbedoDivide, FDimAtlasType, FDimBufferPassThrough>;
 	};
 
 	// Reshape the layout of the buffer to target
@@ -578,7 +580,7 @@ namespace NFORDenoise
 
 			BEGIN_SHADER_PARAMETER_STRUCT(FParameters, )
 				SHADER_PARAMETER_RDG_BUFFER_SRV(Buffer<float>, X)
-				SHADER_PARAMETER_RDG_BUFFER_SRV(Buffer<float>, W)
+				SHADER_PARAMETER_RDG_BUFFER_SRV(Buffer       , W)
 				SHADER_PARAMETER_RDG_BUFFER_SRV(Buffer<float>, Y)
 				SHADER_PARAMETER_RDG_BUFFER_UAV(RWBuffer<float>, Result)
 
@@ -732,7 +734,7 @@ namespace NFORDenoise
 		public:
 			BEGIN_SHADER_PARAMETER_STRUCT(FParameters, )
 				SHADER_PARAMETER_RDG_BUFFER_SRV(Buffer<float>, X)
-				SHADER_PARAMETER_RDG_BUFFER_SRV(Buffer<float>, W)
+				SHADER_PARAMETER_RDG_BUFFER_SRV(Buffer       , W)
 				SHADER_PARAMETER_RDG_BUFFER_SRV(Buffer<float>, B)
 				SHADER_PARAMETER_RDG_TEXTURE_UAV(RWTexture2D<float4>, RWReconstruction)
 				SHADER_PARAMETER_RDG_BUFFER_UAV(RWStructuredBuffer<uint4>, RWReconstructBuffer)
diff --git a/Engine/Plugins/Experimental/NFORDenoise/Source/NFORDenoise/Shaders/NFORDenoise.usf b/Engine/Plugins/Experimental/NFORDenoise/Source/NFORDenoise/Shaders/NFORDenoise.usf
--- a/Engine/Plugins/Experimental/NFORDenoise/Source/NFORDenoise/Shaders/NFORDenoise.usf
+++ b/Engine/Plugins/Experimental/NFORDenoise/Source/NFORDenoise/Shaders/NFORDenoise.usf
@@ -69,16 +69,29 @@
 #error NONLOCALMEAN_ATLAS_TYPE is not supported.
 #endif
 
+#ifndef BUFFER_PASS_THROUGH
+#define BUFFER_PASS_THROUGH 0
+#endif
+
 #define NONLOCALMEAN_ATLAS_SYMMETRIC_PAIR_COUNT (NONLOCALMEAN_ATLAS_TYPE+1)
 
 #define NLM_WEIGHTLAYOUT_NONE						  0
 #define NLM_WEIGHTLAYOUT_NUM_OF_WEIGHTS_PER_PIXELxWxH 1
 #define NLM_WEIGHTLAYOUT_WxHxNUM_OF_WEIGHTS_PER_PIXEL 2
+#define NLM_WEIGHTLAYOUT_4xWxHxNUM_OF_WEIGHTS_BY_4	  3
 
 #ifndef NLM_WEIGHTLAYOUT
 #define NLM_WEIGHTLAYOUT NLM_WEIGHTLAYOUT_NONE
 #endif
 
+#if NLM_WEIGHTLAYOUT == NLM_WEIGHTLAYOUT_4xWxHxNUM_OF_WEIGHTS_BY_4
+#define TWeightReadType float4
+#define WEIGHT_PIXEL_INCREMENT	4
+#else
+#define TWeightReadType float
+#define WEIGHT_PIXEL_INCREMENT	1
+#endif
+
 #ifndef NUM_FEATURE
 #define NUM_FEATURE 0
 #endif
@@ -515,7 +528,7 @@ float GetNonLocalMeanWeight(int2 P, int2 Q)
 }
 
 RWTexture2D<TPixelValue> DenoisedImage;
-Buffer<float> NonLocalMeanWeights;
+Buffer<TWeightReadType> NonLocalMeanWeights;
 int DenoisingChannelCount;
 int4 FilteringRegion;
 
@@ -529,6 +542,11 @@ uint GetWeightBufferIndex(uint2 LocalPositionInRegion, uint WeightIndex, uint2 R
 #elif NLM_WEIGHTLAYOUT == NLM_WEIGHTLAYOUT_WxHxNUM_OF_WEIGHTS_PER_PIXEL
 	// Best if neighbour pixels query the same weight index.
 	uint BufferIndex = LocalPositionInRegion.x + RegionSize.x * (LocalPositionInRegion.y + RegionSize.y * WeightIndex);
+#elif NLM_WEIGHTLAYOUT == NLM_WEIGHTLAYOUT_4xWxHxNUM_OF_WEIGHTS_BY_4
+	// Assume the buffer layout is 4xWxHx[N/4] and the buffer is accessed with Buffer<float>
+	uint StartIndex = WeightIndex / 4;
+	uint WeightOffset = WeightIndex % 4;
+	uint BufferIndex = WeightOffset + 4 *(LocalPositionInRegion.x + RegionSize.x * (LocalPositionInRegion.y + RegionSize.y * StartIndex));
 #else
 	#error not implemented
 #endif
@@ -545,25 +563,41 @@ void NonLocalMeanFilteringCS(in const uint3 DispatchThreadID : SV_DispatchThread
 	}
 
 	const uint PatchWidth = PatchDistance * 2 + 1;
+	const uint NumOfWeightsPerPixel = Pow2(PatchWidth);
 	const uint2 Position = DispatchThreadID.xy + FilteringRegion.xy;
 	float TotalWeights = 0;
 	TPixelValue Result = 0;
 
-	for (int y = -PatchDistance; y <= PatchDistance; ++y)
+	for (int i = 0; i < NumOfWeightsPerPixel; i += WEIGHT_PIXEL_INCREMENT)
 	{
-		for (int x = -PatchDistance; x <= PatchDistance; ++x)
-		{
-			int2 Q = Position + int2(x, y);
+		int x = i % PatchWidth - PatchDistance;
+		int y = i / PatchWidth - PatchDistance;
+		int2 Q = Position + int2(x,y);
 #if NLM_WEIGHTLAYOUT == NLM_WEIGHTLAYOUT_NONE
-			float Weight = GetNonLocalMeanWeight(Position, Q);
+		TWeightReadType Weight = GetNonLocalMeanWeight(Position, Q);
 #else
-			uint WeightIndex = (x + PatchDistance) + (y + PatchDistance) * PatchWidth;
-			uint BufferIndex = GetWeightBufferIndex(DispatchThreadID.xy, WeightIndex, RegionSize);
-			float Weight = NonLocalMeanWeights[BufferIndex];
+		uint WeightIndex = i;// (x + PatchDistance) + (y + PatchDistance) * PatchWidth;
+		uint BufferIndex = GetWeightBufferIndex(DispatchThreadID.xy, WeightIndex, RegionSize) / WEIGHT_PIXEL_INCREMENT;
+		TWeightReadType Weight = NonLocalMeanWeights[BufferIndex];
 #endif
-			//TODO: Improve caching of weights or image value in LDS.
-			Result += Weight * GetImageValue(GetMirroredPosition(Q, TextureSize), TargetImage);
-			TotalWeights += Weight;
+
+		UNROLL
+		for (int j = 0; j < WEIGHT_PIXEL_INCREMENT; ++j)
+		{	
+#if WEIGHT_PIXEL_INCREMENT == 1
+			float LocalWeight = Weight;
+#else
+			int CombinedIndex = i + j; 
+			float LocalWeight = Weight[j];
+			x = CombinedIndex % PatchWidth - PatchDistance;
+			y = CombinedIndex / PatchWidth - PatchDistance;
+			Q = Position + int2(x,y);
+			if (CombinedIndex < NumOfWeightsPerPixel)
+#endif
+			{
+				Result += LocalWeight * GetImageValue(GetMirroredPosition(Q, TextureSize), TargetImage);
+				TotalWeights += LocalWeight;
+			}
 		}
 	}
 
@@ -602,7 +636,8 @@ void NonLocalMeanWeightsCS(in const uint3 DispatchThreadID : SV_DispatchThreadID
 	{
 		int2 Offset = int2(i % PatchDiameter, i / PatchDiameter) - int2(PatchDistance, PatchDistance);
 		int2 Q = Position + Offset;
-#if NLM_WEIGHTLAYOUT != NLM_WEIGHTLAYOUT_WxHxNUM_OF_WEIGHTS_PER_PIXEL
+#if (NLM_WEIGHTLAYOUT == NLM_WEIGHTLAYOUT_NONE) \
+	|| (NLM_WEIGHTLAYOUT == NLM_WEIGHTLAYOUT_NUM_OF_WEIGHTS_PER_PIXELxWxH)
 		uint BufferIndex = BufferOffset + i;
 #else
 		uint BufferIndex = GetWeightBufferIndex(DispatchThreadID.xy, i, RegionSize);
@@ -710,14 +745,27 @@ void NonLocalMeanSeperableFilterPatchSqauredDistanceCS(in const uint3 DispatchTh
 	TAtlasValueType PatchWeight0 = GetNonLocalMeanWeight(AtlasPosition, Q, PatchSquareDistance);
 	
 	// Write to buffer with dimension layout XxYxNumB (each element is of size Count(TAtlasValueType,float)/2)
+#if BUFFER_PASS_THROUGH
+	const uint x = DispatchThreadID.x - (PatchDistance);
+	const uint y = DispatchThreadID.y - (PatchDistance);
+	const uint X = SeperableRegionSize.x - 2 * PatchDistance;
+	const uint Y = SeperableRegionSize.y - 2 * PatchDistance;
+#else
 	const uint x = DispatchThreadID.x;
-	const uint X = SeperableRegionSize.x;
 	const uint y = DispatchThreadID.y;
+	const uint X = SeperableRegionSize.x;
 	const uint Y = SeperableRegionSize.y;
+#endif
 	const uint z = DispatchId * DispatchTileCount + DispatchThreadID.z;
+	const uint SaveIndex = x + X * (y + z * Y);
 
-	const uint SaveIndex =  x + X * (y + z * Y);
-	RWNLMWeights[SaveIndex] = PatchWeight0;
+#if BUFFER_PASS_THROUGH
+	const uint MaxZ = (Pow2(2 * PatchDistance + 1) - 1 + 4) / 4;
+	if (x>=0 && y >=0 && x < X && y < Y && z < MaxZ)
+#endif
+	{
+		RWNLMWeights[SaveIndex] = PatchWeight0;
+	}
 #endif
 }
 
@@ -746,10 +794,15 @@ void WriteWeight(uint2 PixelPosition, uint2 WeightPosition, float Weight)
 	const uint PatchDiemeter = 2 * PatchDistance + 1;
 	uint WeightIndex = WeightPosition.x + WeightPosition.y * PatchDiemeter;
 
-#if NLM_WEIGHTLAYOUT != NLM_WEIGHTLAYOUT_WxHxNUM_OF_WEIGHTS_PER_PIXEL
+#if (NLM_WEIGHTLAYOUT == NLM_WEIGHTLAYOUT_NONE) \
+	|| (NLM_WEIGHTLAYOUT == NLM_WEIGHTLAYOUT_NUM_OF_WEIGHTS_PER_PIXELxWxH)
 	uint OutputIndex = WeightIndex + PixelPosition.x * TargetBufferDim.x + PixelPosition.y * TargetBufferDim.x * TargetBufferDim.y;
-#else
+#elif NLM_WEIGHTLAYOUT == NLM_WEIGHTLAYOUT_WxHxNUM_OF_WEIGHTS_PER_PIXEL
 	uint OutputIndex = PixelPosition.x + TargetBufferDim.y * (PixelPosition.y + TargetBufferDim.z * WeightIndex);
+#elif NLM_WEIGHTLAYOUT == NLM_WEIGHTLAYOUT_4xWxHxNUM_OF_WEIGHTS_BY_4
+	uint StartIndex = WeightIndex / 4;
+	uint WeightOffset = WeightIndex % 4;
+	uint OutputIndex = WeightOffset + 4 * (PixelPosition.x + TargetBufferDim.y * (PixelPosition.y + TargetBufferDim.z * StartIndex));
 #endif
 	RWTargetBuffer[OutputIndex] = Weight;
 }
@@ -875,7 +928,7 @@ void NormalizeTextureCS(in const uint3 DispatchThreadID : SV_DispatchThreadID)
 //	2. Weighted Least-square solver
 
 Buffer<float> X;
-Buffer<float> W;
+Buffer<TWeightReadType> W;
 Buffer<float> Y;
 
 RWBuffer<float> Result;
@@ -929,15 +982,26 @@ float GetY(int2 Position, int f)
 #endif
 }
 
-float GetW(int2 Position, int WeightIndex, int FrameIndex)
+TWeightReadType GetW(int2 Position, int WeightIndex, int FrameIndex)
 {
 	const int NumberOfPixels = TextureSize.x * TextureSize.y;
 
-#if NLM_WEIGHTLAYOUT != NLM_WEIGHTLAYOUT_WxHxNUM_OF_WEIGHTS_PER_PIXEL
+#if (NLM_WEIGHTLAYOUT == NLM_WEIGHTLAYOUT_NONE) \
+	|| (NLM_WEIGHTLAYOUT == NLM_WEIGHTLAYOUT_NUM_OF_WEIGHTS_PER_PIXELxWxH)
 	const int Index = NumOfWeigthsPerPixelPerFrame *(NumberOfPixels * FrameIndex + (Position.y * TextureSize.x + Position.x)) + WeightIndex;
-#else
+
+#elif NLM_WEIGHTLAYOUT == NLM_WEIGHTLAYOUT_WxHxNUM_OF_WEIGHTS_PER_PIXEL
 	const int Index = NumOfWeigthsPerPixelPerFrame * NumberOfPixels * FrameIndex
 					  + Position.x + TextureSize.x * (Position.y + TextureSize.y * WeightIndex);
+
+#elif NLM_WEIGHTLAYOUT == NLM_WEIGHTLAYOUT_4xWxHxNUM_OF_WEIGHTS_BY_4
+	const uint FrameOffset = NumberOfPixels * (NumOfWeigthsPerPixelPerFrame + 4 - 1) / 4;
+	uint StartIndex = WeightIndex / 4; // return the weight stride instead of the specific weight
+	uint WeightOffset = WeightIndex % 4;
+	const uint Index = FrameOffset * FrameIndex
+					   + (Position.x + TextureSize.x * (Position.y + TextureSize.y * StartIndex));
+#else
+	#error not implemented
 #endif
 	return W[Index];
 }
@@ -967,6 +1031,7 @@ int SamplingStep;
 
 void SolveMatrixMultiplicationPerMatrixItem(int2 Position, int2 ModifiedXyYyDim, int PatchWidth, int SingleFrameDataCount)
 {
+#if NLM_WEIGHTLAYOUT != NLM_WEIGHTLAYOUT_4xWxHxNUM_OF_WEIGHTS_BY_4
 	for (int i = 0; i < ModifiedXyYyDim.x; ++i)
 	{
 #if WEIGHTED_MULTIPLICATION_TYPE == WEIGHTED_MULTIPLICATION_QUADRATIC
@@ -1017,6 +1082,7 @@ void SolveMatrixMultiplicationPerMatrixItem(int2 Position, int2 ModifiedXyYyDim,
 #endif
 		}
 	}
+#endif // NLM_WEIGHTLAYOUT != NLM_WEIGHTLAYOUT_4xWxHxNUM_OF_WEIGHTS_BY_4
 }
 
 #if WEIGHTED_MULTIPLICATION_TYPE == WEIGHTED_MULTIPLICATION_QUADRATIC
@@ -1073,7 +1139,7 @@ void MatrixMultiplicationSmallMatrixOptimization(int2 Position, int2 ModifiedXyY
 	int NumOfFrames = XDim.x / SingleFrameDataCount;
 	int Index = 0;
 
-#if USE_SAMPLING_STEP
+#if USE_SAMPLING_STEP && (NLM_WEIGHTLAYOUT != NLM_WEIGHTLAYOUT_4xWxHxNUM_OF_WEIGHTS_BY_4)
 	#if SAMPLING_TYPE == SAMPLING_TYPE_INCREMENTAL
 	//This sampling might lead to incorrect DOF when there is ambiguity using feature to predict
 	//non-DOF background in previous/future frames, and DOF foreground that is the current frame.
@@ -1111,7 +1177,7 @@ void MatrixMultiplicationSmallMatrixOptimization(int2 Position, int2 ModifiedXyY
 		int2 TargetPosition = Position + Offset;
 	#endif
 #else
-	for (int j = 0; j < SingleFrameDataCount; ++j)
+	for (int j = 0; j < SingleFrameDataCount; j += WEIGHT_PIXEL_INCREMENT)
 	{
 		int IndexWithinFrame = j;
 		int2 Offset = int2(IndexWithinFrame % PatchWidth, IndexWithinFrame / PatchWidth);
@@ -1121,33 +1187,49 @@ void MatrixMultiplicationSmallMatrixOptimization(int2 Position, int2 ModifiedXyY
 #endif
 		{
 			int frame_index = t;
-			float sqrtw = sqrt(GetW(Position, IndexWithinFrame, frame_index)); // get w.
+			TWeightReadType sqrtw = sqrt(GetW(Position, IndexWithinFrame, frame_index)); // get w.
 			
+			UNROLL
+			for (int WeightStrideIndex = 0; WeightStrideIndex < WEIGHT_PIXEL_INCREMENT; ++WeightStrideIndex)
+			{
+#if	WEIGHT_PIXEL_INCREMENT == 1
+				float LocalWeight = sqrtw;
+#else
+				IndexWithinFrame = j + WeightStrideIndex;
+				Offset = int2(IndexWithinFrame % PatchWidth, IndexWithinFrame / PatchWidth);
+				TargetPosition = Position + Offset;
+				float LocalWeight = sqrtw[WeightStrideIndex];
+
+				if (IndexWithinFrame < SingleFrameDataCount)
+#endif
+				{
 #if WEIGHTED_MULTIPLICATION_TYPE == WEIGHTED_MULTIPLICATION_GENERALIZED
-			FFeatureYCache Y;
-			FillFeatureYCache(Y, TargetPosition, ModifiedXyYyDim, YDim.y * frame_index);
+					FFeatureYCache Y;
+					FillFeatureYCache(Y, TargetPosition, ModifiedXyYyDim, YDim.y * frame_index);
 #endif
-			FFeatureXCache X;
-			FillFeatureXCache(X, TargetPosition, ModifiedXyYyDim, XDim.y * frame_index);
+					FFeatureXCache X;
+					FillFeatureXCache(X, TargetPosition, ModifiedXyYyDim, XDim.y * frame_index);
 
-			Index = 0;
-			for (int i = 0; i < ModifiedXyYyDim.x; ++i)
-			{
+					Index = 0;
+					for (int i = 0; i < ModifiedXyYyDim.x; ++i)
+					{
 #if WEIGHTED_MULTIPLICATION_TYPE == WEIGHTED_MULTIPLICATION_QUADRATIC
-				for (int k = i; k < ModifiedXyYyDim.y; ++k)
+						for (int k = i; k < ModifiedXyYyDim.y; ++k)
 #else
-				for (int k = 0; k < ModifiedXyYyDim.y; ++k)
+						for (int k = 0; k < ModifiedXyYyDim.y; ++k)
 #endif
-				{
-					float x = X.Data[i];
+						{
+							float x = X.Data[i];
 
 #if WEIGHTED_MULTIPLICATION_TYPE == WEIGHTED_MULTIPLICATION_GENERALIZED
-					float y = Y.Data[k];
+							float y = Y.Data[k];
 #else 
-					float y = X.Data[k];
+							float y = X.Data[k];
 #endif
-					MatrixCache.Data[Index] += x * sqrtw * y;
-					++Index;
+							MatrixCache.Data[Index] += x * LocalWeight * y;
+							++Index;
+						}
+					}
 				}
 			}
 		}
@@ -1409,22 +1491,34 @@ void ReconstructSpatialTemporalImageCS(in const uint3 DispatchThreadID : SV_Disp
 #endif
 
 	// For center frame, reconstruct and scatter the result
-	for (int y = 0; y < PatchSideLength; ++y)
+	for (int PixelIndex = 0; PixelIndex < WeightStride; PixelIndex += WEIGHT_PIXEL_INCREMENT)
 	{
-		for (int x = 0; x < PatchSideLength; ++x)
+		int x = PixelIndex % PatchSideLength;
+		int y = PixelIndex / PatchSideLength;
+		TWeightReadType PixelWeights = GetW(Position, PixelIndex, FrameIndex);
+		
+		UNROLL
+		for (int SubStrideIndex = 0; SubStrideIndex < WEIGHT_PIXEL_INCREMENT; ++SubStrideIndex)
 		{
-			int WeightIndex1D = (x + y * PatchSideLength);
-			float w = GetW(Position, WeightIndex1D, FrameIndex);
-
-			uint2 TargetPosition = Position + int2(x, y);
-			float3 Denoised = 0;
-			
-			for (int i = 0; i < NUM_FEATURE; ++i)
+#if WEIGHT_PIXEL_INCREMENT == 1
+			float LocalWeight = PixelWeights;
+#else
+			int ResolvedPixelIndex = PixelIndex + SubStrideIndex;
+			x = ResolvedPixelIndex % PatchSideLength;
+			y = ResolvedPixelIndex / PatchSideLength;
+			float LocalWeight = PixelWeights[SubStrideIndex];
+			if (ResolvedPixelIndex < WeightStride)
+#endif
 			{
-				float FeatureI = 1.0f; 
-				if (i < NumOfFeatures)
+				uint2 TargetPosition = Position + int2(x, y);
+				float3 Denoised = 0;
+
+				for (int i = 0; i < NUM_FEATURE; ++i)
 				{
-					FeatureI = GetX(TargetPosition, i, FrameIndex);
+					float FeatureI = 1.0f;
+					if (i < NumOfFeatures)
+					{
+						FeatureI = GetX(TargetPosition, i, FrameIndex);
 
 #if PRE_ALBEDO_DIVIDE == PRE_ALBEDO_DIVIDE_EACH
 						//Assume the first three channels are albedo. TODO: use passed in index
@@ -1433,28 +1527,29 @@ void ReconstructSpatialTemporalImageCS(in const uint3 DispatchThreadID : SV_Disp
 							Albedo[i] = FeatureI;
 						}
 #endif
+					}
+					Denoised += FeatureI * Weights[i];
 				}
-				Denoised += FeatureI * Weights[i];
-			}
 
 #if RECONSTRUCTION_TYPE == RECONSTRUCTION_TYPE_SCATTER
 
 #if PRE_ALBEDO_DIVIDE == PRE_ALBEDO_DIVIDE_EACH
-			Denoised *= RADIANCE_PREPROCESS_SCALE_FACTOR * Albedo;
+				Denoised *= RADIANCE_PREPROCESS_SCALE_FACTOR * Albedo;
 #elif PRE_ALBEDO_DIVIDE == PRE_ALBEDO_DIVIDE_FINAL
-			Denoised *= RADIANCE_PREPROCESS_SCALE_FACTOR;
+				Denoised *= RADIANCE_PREPROCESS_SCALE_FACTOR;
 #endif
-			// Scatter the result instead of gathering for the current frame to get high frequency results (better reconstruction).
-			float4 XBW = float4(Denoised * w, w);
-			SaveFloat4ToBufferOrTexture(TargetPosition, XBW, WeightStride);
+				// Scatter the result instead of gathering for the current frame to get high frequency results (better reconstruction).
+				float4 XBW = float4(Denoised * LocalWeight, LocalWeight);
+				SaveFloat4ToBufferOrTexture(TargetPosition, XBW, WeightStride);
 
 #elif RECONSTRUCTION_TYPE == RECONSTRUCTION_TYPE_GATHER
 
 #if PRE_ALBEDO_DIVIDE == PRE_ALBEDO_DIVIDE_EACH
-			Denoised *= Albedo;
+				Denoised *= Albedo;
 #endif
-			GatheredValue += float4(Denoised * w, w);
+				GatheredValue += float4(Denoised * LocalWeight, LocalWeight);
 #endif
+			}
 		}
 	}
 #if RECONSTRUCTION_TYPE == RECONSTRUCTION_TYPE_GATHER
