Add radiance alpha channel denoising through NLM with alpha variance (r.NFOR.Denoise.Alpha true). It's on by default.
[RN] Variance texture stores the alpha variance in the A component.
#tests Renders as expected on hillside project.

diff --git a/Engine/Plugins/Experimental/NFORDenoise/Source/NFORDenoise/Private/NFORDenoiseCS.cpp b/Engine/Plugins/Experimental/NFORDenoise/Source/NFORDenoise/Private/NFORDenoiseCS.cpp
--- a/Engine/Plugins/Experimental/NFORDenoise/Source/NFORDenoise/Private/NFORDenoiseCS.cpp
+++ b/Engine/Plugins/Experimental/NFORDenoise/Source/NFORDenoise/Private/NFORDenoiseCS.cpp
@@ -219,6 +219,12 @@ namespace NFORDenoise
 		TEXT("The patch distance for bandwidth selection dependents on this parameters for MSE and selection filtering."),
 		ECVF_RenderThreadSafe);
 
+	TAutoConsoleVariable<bool> CVarNFORDenoiseAlpha(
+		TEXT("r.NFOR.Denoise.Alpha"),
+		true,
+		TEXT("Indicate if the alpha channel of radiance will be denoised (Default on)."),
+		ECVF_RenderThreadSafe);
+
 	TAutoConsoleVariable<bool> CVarNFORBandwidthSelection(
 		TEXT("r.NFOR.BandwidthSelection"),
 		true,
@@ -539,6 +545,11 @@ namespace NFORDenoise
 		return NonLocalMeanAtlasType;
 	}
 
+	bool ShouldDenoiseAlpha()
+	{
+		return CVarNFORDenoiseAlpha.GetValueOnRenderThread();
+	}
+
 	bool IsBandwidthSelectionEnabled()
 	{
 		return CVarNFORBandwidthSelection.GetValueOnRenderThread();
@@ -584,6 +595,7 @@ namespace NFORDenoise
 	IMPLEMENT_GLOBAL_SHADER(FTextureAccumulateConstantCS, "/NFORDenoise/NFORDenoise.usf", "TextureOperationCS", SF_Compute);
 	IMPLEMENT_GLOBAL_SHADER(FTextureAccumulateCS, "/NFORDenoise/NFORDenoise.usf", "TextureOperationCS", SF_Compute);
 	IMPLEMENT_GLOBAL_SHADER(FCopyTexturePS, "/NFORDenoise/NFORDenoise.usf", "CopyTexturePS", SF_Pixel);
+	IMPLEMENT_GLOBAL_SHADER(FCopyTextureSingleChannelCS, "/NFORDenoise/NFORDenoise.usf", "CopyTextureSingleChannelCS", SF_Compute);
 
 	//--------------------------------------------------------------------------------------------------------------------
 	// Feature range adjustment and radiance normalization
@@ -816,6 +828,50 @@ namespace NFORDenoise
 		);
 	}
 
+	void AddCopyMirroredTexturePass(FRDGBuilder& GraphBuilder, const FRDGTextureRef& SourceTexture, const FRDGTextureRef& TargetTexture, int32 Channel, ETextureCopyType CopyType,
+		FIntPoint SourcePosition, FIntPoint TargetPosition, FIntPoint Size)
+	{
+		const FIntPoint CopySize = Size == FIntPoint::ZeroValue ? TargetTexture->Desc.Extent : Size;
+
+		FGlobalShaderMap* ShaderMap = GetGlobalShaderMap(GMaxRHIFeatureLevel);
+
+		typedef FCopyTextureSingleChannelCS SHADER;
+		SHADER::FParameters* PassParameters = GraphBuilder.AllocParameters<SHADER::FParameters>();
+		{
+			PassParameters->CopySource = GraphBuilder.CreateSRV(FRDGTextureSRVDesc::Create(SourceTexture));
+			PassParameters->RWCopyTarget = GraphBuilder.CreateUAV(TargetTexture);
+			PassParameters->SourceOffset = SourcePosition;
+			PassParameters->TargetOffset = TargetPosition;
+			PassParameters->CopySize = CopySize;
+			PassParameters->Channel = Channel;
+			PassParameters->TextureSize = SourceTexture->Desc.Extent;
+		}
+
+		SHADER::FPermutationDomain PermutationDomainVector;
+		{
+			PermutationDomainVector.Set<SHADER::FDimTextureCopyType>(CopyType);
+		}
+
+		TShaderMapRef<SHADER> ComputeShader(ShaderMap, PermutationDomainVector);
+
+		FComputeShaderUtils::AddPass(
+			GraphBuilder,
+			RDG_EVENT_NAME("CopyTextureCS (%s [%d,%d] -> %s [%d,%d], size:%dx%d, c=%d)",
+				SourceTexture->Name,
+				SourcePosition.X,
+				SourcePosition.Y,
+				TargetTexture->Name,
+				TargetPosition.X,
+				TargetPosition.Y,
+				CopySize.X,
+				CopySize.Y,
+				Channel),
+			ERDGPassFlags::Compute,
+			ComputeShader,
+			PassParameters,
+			FComputeShaderUtils::GetGroupCount(CopySize, NON_LOCAL_MEAN_THREAD_GROUP_SIZE));
+	}
+
 	void AddNormalizeRadianceVariancePass(FRDGBuilder& GraphBuilder, const FRDGTextureRef& Albedo, const FRDGTextureRef& RadianceVariance)
 	{
 		FIntPoint Size = RadianceVariance->Desc.Extent;
@@ -2423,10 +2479,6 @@ namespace NFORDenoise
 		// Normalize the image by weights stored in alpha channel.
 		AddNormalizeTexturePass(GraphBuilder, FilteredRadiance);
 
-		// TODO: denoise the alpha channel of the original radiance.
-		// Pass through alpha channel from the source index.
-		AddCopyMirroredTexturePass(GraphBuilder, Radiances[SourceIndex].Data.Image, FilteredRadiance, FIntPoint::ZeroValue, FIntPoint::ZeroValue, FIntPoint::ZeroValue,true/*bAlphaOnly*/);
-
 		return FilteredRadiance;
 	}
 
@@ -2714,6 +2766,51 @@ namespace NFORDenoise
 			AddCopyTexturePass(GraphBuilder, FilteredImages[0], DenoisedRadiance);
 		}
 
+		{
+			RDG_EVENT_SCOPE(GraphBuilder, "Postprocessing");
+
+			// Alpha denoising
+			{
+				FRDGTextureRef RadianceAlpha = Radiances[SourceRadianceIndex].Data.Image;
+				if (ShouldDenoiseAlpha())
+				{
+					// Apply non-local mean filter to alpha only. 
+					FRDGTextureRef RadianceAlphaVariance = Radiances[SourceRadianceIndex].Variance.Image;
+					EVarianceType VarianceType = Radiances[SourceRadianceIndex].VarianceType;
+					FRDGTextureDesc AlphaTextureDesc = RadianceAlpha->Desc;
+					AlphaTextureDesc.Format = PF_R32_FLOAT;
+					FRDGTextureRef RawAlphaTexture = GraphBuilder.CreateTexture(AlphaTextureDesc, TEXT("NFOR.RawAlphaTexture"));
+					FRDGTextureRef FilteredAlphaTexture = GraphBuilder.CreateTexture(AlphaTextureDesc, TEXT("NFOR.DenoisedAlphaTexture"));
+
+					// Assume alpha is the a component of Radiance texture, the alpha variance is the a component of the corresponding variance texture.
+					const int32 AlphaChannelIndex = 3;
+
+					AddCopyMirroredTexturePass(GraphBuilder, RadianceAlpha, RawAlphaTexture, AlphaChannelIndex, ETextureCopyType::TargetSingleChannel);
+
+					FNFORTextureDesc AlphaTexture = FNFORTextureDesc(RawAlphaTexture, 0, 1, 1);
+					FNonLocalMeanParameters AlphaNLMParams = GetFeatureNonLocalMeanParameters(0.5f);
+					FNFORTextureDesc AlphaVariance = FNFORTextureDesc(RadianceAlphaVariance, AlphaChannelIndex, 1, 4);
+
+					ApplyNonLocalMeanFilterIfRequired(
+						GraphBuilder,
+						View,
+						AlphaNLMParams,
+						AlphaTexture,
+						AlphaVariance,
+						VarianceType,
+						FilteredAlphaTexture,
+						GetFeatureTileSizeDownScale());
+
+					AddCopyMirroredTexturePass(GraphBuilder, FilteredAlphaTexture, DenoisedRadiance, AlphaChannelIndex, ETextureCopyType::SourceSingleChannel);
+				}
+				else
+				{
+					// Pass through alpha channel.
+					AddCopyMirroredTexturePass(GraphBuilder, RadianceAlpha, DenoisedRadiance, FIntPoint::ZeroValue, FIntPoint::ZeroValue, FIntPoint::ZeroValue, true/*bAlphaOnly*/);
+				}
+			}
+		}
+
 		return true;
 	}
 }
diff --git a/Engine/Plugins/Experimental/NFORDenoise/Source/NFORDenoise/Private/NFORDenoiseCS.h b/Engine/Plugins/Experimental/NFORDenoise/Source/NFORDenoise/Private/NFORDenoiseCS.h
--- a/Engine/Plugins/Experimental/NFORDenoise/Source/NFORDenoise/Private/NFORDenoiseCS.h
+++ b/Engine/Plugins/Experimental/NFORDenoise/Source/NFORDenoise/Private/NFORDenoiseCS.h
@@ -164,6 +164,38 @@ namespace NFORDenoise
 		END_SHADER_PARAMETER_STRUCT()
 	};
 
+	enum class ETextureCopyType :uint32
+	{
+		TargetSingleChannel,
+		SourceSingleChannel,
+		MAX
+	};
+
+	class FCopyTextureSingleChannelCS : public FGlobalShader
+	{
+		DECLARE_GLOBAL_SHADER(FCopyTextureSingleChannelCS)
+		SHADER_USE_PARAMETER_STRUCT(FCopyTextureSingleChannelCS, FGlobalShader)
+	public:
+		BEGIN_SHADER_PARAMETER_STRUCT(FParameters, )
+			SHADER_PARAMETER_RDG_TEXTURE_SRV(Texture2D, CopySource)
+			SHADER_PARAMETER_RDG_TEXTURE_UAV(RWTexture2D, RWCopyTarget)
+			SHADER_PARAMETER(FIntPoint, SourceOffset)
+			SHADER_PARAMETER(FIntPoint, TargetOffset)
+			SHADER_PARAMETER(FIntPoint, CopySize)
+			SHADER_PARAMETER(FIntPoint, TextureSize)
+			SHADER_PARAMETER(int32, Channel)
+		END_SHADER_PARAMETER_STRUCT()
+
+		static void ModifyCompilationEnvironment(const FGlobalShaderPermutationParameters& InParameters, FShaderCompilerEnvironment& OutEnvironment)
+		{
+			FGlobalShader::ModifyCompilationEnvironment(InParameters, OutEnvironment);
+			OutEnvironment.SetDefine(TEXT("THREAD_GROUP_SIZE"), NON_LOCAL_MEAN_THREAD_GROUP_SIZE);
+		}
+
+		class FDimTextureCopyType : SHADER_PERMUTATION_ENUM_CLASS("TEXTURE_COPY_TYPE", ETextureCopyType);
+		using FPermutationDomain = TShaderPermutationDomain<FDimTextureCopyType>;
+	};
+
 	// TargetTexture.rgb = TargetTexture.rgb * lerp (1.0f,SourceTexture.rgb, (SourceTexture.rgb != 0 || bForceMultiply));
 	void AddMultiplyTextureRegionPass(FRDGBuilder& GraphBuilder, const FRDGTextureRef& SourceTexture, const FRDGTextureRef& TargetTexture,
 		bool bForceMultiply = true, FIntPoint SourcePosition = FIntPoint::ZeroValue, FIntPoint TargetPosition = FIntPoint::ZeroValue, FIntPoint Size = FIntPoint::ZeroValue);
@@ -184,6 +216,9 @@ namespace NFORDenoise
 	void AddCopyMirroredTexturePass(FRDGBuilder& GraphBuilder, const FRDGTextureRef& SourceTexture, const FRDGTextureRef& TargetTexture,
 		FIntPoint SourcePosition = FIntPoint::ZeroValue, FIntPoint TargetPosition = FIntPoint::ZeroValue, FIntPoint Size = FIntPoint::ZeroValue, bool bAlphaOnly = false);
 	
+	// Copy a channel from the source image to the target image where ETextureCopyType indicates which image is single channel;
+	void AddCopyMirroredTexturePass(FRDGBuilder& GraphBuilder, const FRDGTextureRef& SourceTexture, const FRDGTextureRef& TargetTexture, int32 Channel, ETextureCopyType CopyType,
+		FIntPoint SourcePosition = FIntPoint::ZeroValue, FIntPoint TargetPosition = FIntPoint::ZeroValue, FIntPoint Size = FIntPoint::ZeroValue);
 
 	//--------------------------------------------------------------------------------------------------------------------
 	// Feature range adjustment, and radiance normalization by albedo
diff --git a/Engine/Plugins/Experimental/NFORDenoise/Source/NFORDenoise/Shaders/NFORDenoise.usf b/Engine/Plugins/Experimental/NFORDenoise/Source/NFORDenoise/Shaders/NFORDenoise.usf
--- a/Engine/Plugins/Experimental/NFORDenoise/Source/NFORDenoise/Shaders/NFORDenoise.usf
+++ b/Engine/Plugins/Experimental/NFORDenoise/Source/NFORDenoise/Shaders/NFORDenoise.usf
@@ -156,6 +156,13 @@
 #define TEXTURE_OPS TEXTURE_OPS_MULTIPLY
 #endif
 
+#define TEXTURE_COPY_TARGET_SINGLE_CHANNEL	0
+#define TEXTURE_COPY_SOURCE_SINGLE_CHANNEL	1
+
+#ifndef TEXTURE_COPY_TYPE
+#define TEXTURE_COPY_TYPE TEXTURE_COPY_TARGET_SINGLE_CHANNEL
+#endif
+
 #ifndef ACCUMULATE_BY_MASK
 #define ACCUMULATE_BY_MASK 0
 #endif
@@ -247,12 +254,45 @@ void TextureOperationCS(in const uint3 DispatchThreadID : SV_DispatchThreadID)
 }
 
 int2 SourceOffset;
+int2 TargetOffset;
+int Channel;
+int2 CopySize;
+
 float4 CopyTexturePS(float4 SvPosition : SV_POSITION) : SV_Target0
 {
 	uint2 Position = (uint2)SvPosition.xy + SourceOffset;
 	return Source.Load(uint3(GetMirroredPosition(Position, TextureSize), 0));
 }
 
+#if TEXTURE_COPY_TYPE == TEXTURE_COPY_TARGET_SINGLE_CHANNEL
+Texture2D<float4> CopySource;
+RWTexture2D<float> RWCopyTarget;
+#elif TEXTURE_COPY_TYPE == TEXTURE_COPY_SOURCE_SINGLE_CHANNEL
+Texture2D<float> CopySource;
+RWTexture2D<float4> RWCopyTarget;
+#endif
+
+[numthreads(THREAD_GROUP_SIZE, THREAD_GROUP_SIZE, 1)]
+void CopyTextureSingleChannelCS(in const uint3 DispatchThreadID : SV_DispatchThreadID)
+{
+	if (any(DispatchThreadID.xy >= (uint2)CopySize))
+	{
+		return;
+	}
+
+	uint2 ResolvedSourcePosition = DispatchThreadID.xy + SourceOffset;
+	uint2 ResolvedTargetPosition = DispatchThreadID.xy + TargetOffset;
+
+	if (all(ResolvedTargetPosition >= 0))
+	{
+#if TEXTURE_COPY_TYPE == TEXTURE_COPY_TARGET_SINGLE_CHANNEL
+		RWCopyTarget[ResolvedTargetPosition] = CopySource.Load(uint3(GetMirroredPosition(ResolvedSourcePosition, TextureSize), 0))[Channel];
+#elif TEXTURE_COPY_TYPE == TEXTURE_COPY_SOURCE_SINGLE_CHANNEL
+		RWCopyTarget[ResolvedTargetPosition][Channel] = CopySource.Load(uint3(GetMirroredPosition(ResolvedSourcePosition, TextureSize), 0));
+#endif
+	}
+}
+
 //--------------------------------------------------------------------------------------------------
 // Variance definition and functions
 //
diff --git a/Engine/Shaders/Private/PathTracing/PathTracingSpatialTemporalDenoising.usf b/Engine/Shaders/Private/PathTracing/PathTracingSpatialTemporalDenoising.usf
--- a/Engine/Shaders/Private/PathTracing/PathTracingSpatialTemporalDenoising.usf
+++ b/Engine/Shaders/Private/PathTracing/PathTracingSpatialTemporalDenoising.usf
@@ -814,6 +814,9 @@ float4 GetMeanFromTextures(int2 Position)
 	{
 		Value.z = 1.0f;
 	}
+	
+	// Alpha channel needs denoising as well.
+	Value.w = InputTexture[Position].a;
 #else
 	#error "Not implemented" 
 #endif
@@ -835,7 +838,6 @@ void TemporalPrepassCS(uint3 DT_ID : SV_DispatchThreadID)
 
 		FPixelMaterialLightingFingerprint VarianceInfo = (FPixelMaterialLightingFingerprint)0;
 		VarianceInfo.Mean = Mean;
-		VarianceInfo.Mean.w = 0;
 		RWVarianceMap[index] = VarianceInfo;
 
 #elif PREPASS_PHASE == PREPASS_PHASE_UPDATE
@@ -846,7 +848,6 @@ void TemporalPrepassCS(uint3 DT_ID : SV_DispatchThreadID)
 		float4 CurrentValue = (Mean - PreviousMean) * Iteration + Mean;
 		
 		VarianceInfo.Mean = Mean;
-		VarianceInfo.Mean.w = Iteration;
 		VarianceInfo.Var += (CurrentValue - VarianceInfo.Mean)*(CurrentValue - PreviousMean);
 
 		RWVarianceMap[index] = VarianceInfo;
@@ -860,6 +861,7 @@ void TemporalPrepassCS(uint3 DT_ID : SV_DispatchThreadID)
 #if PREPASS_GENERATE_TEXTURE
 RWTexture2D<float4> OutputTexture;
 StructuredBuffer<FPixelMaterialLightingFingerprint> VarianceMap;
+int Iteration;
 
 [numthreads(THREAD_SIZE_X, THREAD_SIZE_Y, 1)]
 void PrepassGenerateTextureCS(uint3 DT_ID : SV_DispatchThreadID)
@@ -873,9 +875,9 @@ void PrepassGenerateTextureCS(uint3 DT_ID : SV_DispatchThreadID)
 
 		// VarianceInfo.Var = sigma^2 * n.
 		// Std of mean = sqrt(sigma^2/n) = sqrt(VarianceInfo.Var/n^2)
-		float4 StandDerivation= sqrt(VarianceInfo.Var) / (VarianceInfo.Mean.w + 1);
+		float4 StandDerivation= sqrt(VarianceInfo.Var) / Iteration;
 
-		OutputTexture[Position] = float4(StandDerivation.xyz, 1.0f);
+		OutputTexture[Position] = StandDerivation;
 	}
 }
 #endif // PREPASS_GENERATE_TEXTURE
diff --git a/Engine/Source/Runtime/Renderer/Private/PathTracingSpatialTemporalDenoising.cpp b/Engine/Source/Runtime/Renderer/Private/PathTracingSpatialTemporalDenoising.cpp
--- a/Engine/Source/Runtime/Renderer/Private/PathTracingSpatialTemporalDenoising.cpp
+++ b/Engine/Source/Runtime/Renderer/Private/PathTracingSpatialTemporalDenoising.cpp
@@ -1667,6 +1667,7 @@ public:
 		SHADER_PARAMETER_RDG_BUFFER_SRV(StructuredBuffer<FPixelMaterialLightingFingerprint>, VarianceMap)
 		SHADER_PARAMETER_RDG_TEXTURE_UAV(RWTexture2D, OutputTexture)
 		SHADER_PARAMETER_STRUCT(FScreenPassTextureViewportParameters, TargetViewport)
+		SHADER_PARAMETER(int32, Iteration)
 	END_SHADER_PARAMETER_STRUCT()
 
 	class FVarianceType : SHADER_PERMUTATION_ENUM_CLASS("VARIANCE_TYPE", FTemporalPrepassCS::EVarianceType);
@@ -1761,6 +1762,7 @@ void PathTracingSpatialTemporalDenoisingPrePass(FRDGBuilder& GraphBuilder, const
 				PassParameters->OutputTexture = GraphBuilder.CreateUAV(SpatialTemporalDenoisingContext.VarianceTexture);
 				PassParameters->VarianceMap	  = GraphBuilder.CreateSRV(VarianceBuffer, EPixelFormat::PF_R32_FLOAT);
 				PassParameters->TargetViewport = TargetViewportParameters;
+				PassParameters->Iteration = IterationNumber;
 			}
 
 			SHADER::FPermutationDomain ComputeShaderPermutationVector;
