Add better description for NFOR denoiser.

#tests No test is requried
#rnx

diff --git a/Engine/Plugins/Experimental/NFORDenoise/NFORDenoise.uplugin b/Engine/Plugins/Experimental/NFORDenoise/NFORDenoise.uplugin
--- a/Engine/Plugins/Experimental/NFORDenoise/NFORDenoise.uplugin
+++ b/Engine/Plugins/Experimental/NFORDenoise/NFORDenoise.uplugin
@@ -3,7 +3,7 @@
 	"Version": 1,
 	"VersionName": "1.0",
 	"FriendlyName": "NFORDenoise",
-	"Description": "Denoising engine for the Unreal Path Tracer based on Nonlinearly Weighted First-order Regression for Denoising Monte Carlo Renderings.",
+	"Description": "Spatial-temporal denoising engine for the Unreal Path Tracer (mainly used with MRQ). It denoises each pixel based on the surrounding patches in space and time in all directions. The algorithm is mainly inspired by Nonlinearly Weighted First-order Regression (NFOR) for Denoising Monte Carlo Renderings.",
 	"Category": "Denoising",
 	"CreatedBy": "Epic Games, Inc.",
 	"CreatedByURL": "https://epicgames.com",
