Modify NFOR to use separated depth and albedo/normal in fp16.
Since depth is passed in through a new variable DepthTex. Denoisers need to make the corresponding changes if Depth is used.

#rnx
#tests MRQ test on hillside project does not show artifacts

diff --git a/Engine/Plugins/Experimental/NFORDenoise/Source/NFORDenoise/Private/NFORDenoise.cpp b/Engine/Plugins/Experimental/NFORDenoise/Source/NFORDenoise/Private/NFORDenoise.cpp
--- a/Engine/Plugins/Experimental/NFORDenoise/Source/NFORDenoise/Private/NFORDenoise.cpp
+++ b/Engine/Plugins/Experimental/NFORDenoise/Source/NFORDenoise/Private/NFORDenoise.cpp
@@ -34,6 +34,7 @@ struct FRDGFrameCache
 	FRDGTextureRef RadianceVariance;
 	FRDGTextureRef Albedo;
 	FRDGTextureRef Normal;
+	FRDGTextureRef Depth;
 
 	FExternalFrameCache ToExternalResource(FRDGBuilder& GraphBuilder, const FSceneView& View);
 	void AppendToNFORBuffer(TArray<NFORDenoise::FRadianceDesc>& Radiances, TArray<NFORDenoise::FFeatureDesc>& Features);
@@ -45,6 +46,7 @@ struct FExternalFrameCache
 	TRefCountPtr<IPooledRenderTarget> RadianceVariance;
 	TRefCountPtr<IPooledRenderTarget> Albedo;
 	TRefCountPtr<IPooledRenderTarget> Normal;
+	TRefCountPtr<IPooledRenderTarget> Depth;
 
 	FRDGFrameCache ToRGDResource(FRDGBuilder& GraphBuilder, const FSceneView& View);
 };
@@ -56,6 +58,7 @@ FExternalFrameCache FRDGFrameCache::ToExternalResource(FRDGBuilder& GraphBuilder
 	ExternalFrameCache.RadianceVariance = GraphBuilder.ConvertToExternalTexture(RadianceVariance);
 	ExternalFrameCache.Albedo = GraphBuilder.ConvertToExternalTexture(Albedo);
 	ExternalFrameCache.Normal = GraphBuilder.ConvertToExternalTexture(Normal);
+	ExternalFrameCache.Depth = GraphBuilder.ConvertToExternalTexture(Depth);
 
 	return ExternalFrameCache;
 }
@@ -82,7 +85,7 @@ void FRDGFrameCache::AppendToNFORBuffer(
 
 		if (IsFeatureDepthEnabled())
 		{
-			NFORDenoise::FNFORTextureDesc DepthTex(Normal, 3, 1, 4);
+			NFORDenoise::FNFORTextureDesc DepthTex(Depth, 0, 1, 1);
 			NFORDenoise::FNFORTextureDesc DepthVarianceTex(nullptr);
 			Features.Add(NFORDenoise::FFeatureDesc(DepthTex, DepthVarianceTex, NFORDenoise::EVarianceType::GreyScale, true));
 		}
@@ -96,6 +99,7 @@ FRDGFrameCache FExternalFrameCache::ToRGDResource(FRDGBuilder& GraphBuilder, con
 	RDGFrameCache.RadianceVariance = GraphBuilder.RegisterExternalTexture(RadianceVariance);
 	RDGFrameCache.Albedo = GraphBuilder.RegisterExternalTexture(Albedo);
 	RDGFrameCache.Normal = GraphBuilder.RegisterExternalTexture(Normal);
+	RDGFrameCache.Depth = GraphBuilder.RegisterExternalTexture(Depth);
 
 	return RDGFrameCache;
 }
@@ -131,11 +135,13 @@ public:
 			RDGFrameCache.RadianceVariance = GraphBuilder.CreateTexture(Inputs.VarianceTex->Desc, TEXT("RadianceVariance"));
 			RDGFrameCache.Albedo = GraphBuilder.CreateTexture(Inputs.AlbedoTex->Desc, TEXT("Albedo"));
 			RDGFrameCache.Normal = GraphBuilder.CreateTexture(Inputs.NormalTex->Desc, TEXT("Normal"));
+			RDGFrameCache.Depth = GraphBuilder.CreateTexture(Inputs.DepthTex->Desc, TEXT("Depth"));
 
 			AddCopyTexturePass(GraphBuilder, Inputs.ColorTex, RDGFrameCache.Radiance);
 			AddCopyTexturePass(GraphBuilder, Inputs.VarianceTex, RDGFrameCache.RadianceVariance);
 			AddCopyTexturePass(GraphBuilder, Inputs.AlbedoTex, RDGFrameCache.Albedo);
 			AddCopyTexturePass(GraphBuilder, Inputs.NormalTex, RDGFrameCache.Normal);
+			AddCopyTexturePass(GraphBuilder, Inputs.DepthTex, RDGFrameCache.Depth);
 			
 			FExternalFrameCache FrameCache = RDGFrameCache.ToExternalResource(GraphBuilder, View);
 
@@ -164,8 +170,8 @@ public:
 		FRDGTextureRef GetDepth(FRDGBuilder& GraphBuilder, int32 FrameIndex)
 		{
 			check(FrameIndex >=0 && FrameIndex < Num());
-			TRefCountPtr<IPooledRenderTarget> NormalDepth= FrameCaches[FrameIndex].Normal;
-			return NormalDepth ? GraphBuilder.RegisterExternalTexture(NormalDepth) : nullptr;
+			TRefCountPtr<IPooledRenderTarget> Depth = FrameCaches[FrameIndex].Depth;
+			return Depth ? GraphBuilder.RegisterExternalTexture(Depth) : nullptr;
 		}
 
 		int32 Num() const { return FrameCaches.Num(); }
@@ -229,8 +235,7 @@ public:
 				const bool bAlphaOnly = true;
 				FRDGTextureRef DenoisingDepth = CurHistory->GetDepth(GraphBuilder, DenoisingDepthFrameIndex);
 
-				NFORDenoise::AddCopyMirroredTexturePass(GraphBuilder, DenoisingDepth, Inputs.NormalTex,
-					FIntPoint::ZeroValue, FIntPoint::ZeroValue, FIntPoint::ZeroValue, bAlphaOnly);
+				AddCopyTexturePass(GraphBuilder, DenoisingDepth, Inputs.DepthTex);
 			}
 		}
 
diff --git a/Engine/Plugins/Experimental/NFORDenoise/Source/NFORDenoise/Private/NFORDenoiseCS.cpp b/Engine/Plugins/Experimental/NFORDenoise/Source/NFORDenoise/Private/NFORDenoiseCS.cpp
--- a/Engine/Plugins/Experimental/NFORDenoise/Source/NFORDenoise/Private/NFORDenoiseCS.cpp
+++ b/Engine/Plugins/Experimental/NFORDenoise/Source/NFORDenoise/Private/NFORDenoiseCS.cpp
@@ -194,6 +194,13 @@ namespace NFORDenoise
 		TEXT(" 3: float4 x Width x Height x DivideAndRoundUp(NumOfWeightsPerPixel,float4)\n"),
 		ECVF_RenderThreadSafe);
 
+	TAutoConsoleVariable<int32> CVarNFORNOnLocalMeanFeatureFormat(
+		TEXT("r.NFOR.NonLocalMean.Feature.Format"),
+		1,
+		TEXT("0: fp32 \n")
+		TEXT("1: fp16 \n"),
+		ECVF_RenderThreadSafe);
+
 	TAutoConsoleVariable<int32> CVarNFORNonLocalMeanFeaturePatchSize(
 		TEXT("r.NFOR.NonLocalMean.Feature.PatchSize"),
 		3,
@@ -569,6 +576,34 @@ namespace NFORDenoise
 		return Bandwidths;
 	}
 
+	EPixelFormat GetFeaturePixelFormat()
+	{
+		EPixelFormat PixelFormat = PF_R32_FLOAT;
+
+		if (GetRegressionDevice() != ERegressionDevice::CPU)
+		{
+			if (CVarNFORNOnLocalMeanFeatureFormat.GetValueOnRenderThread() != 0)
+			{
+				PixelFormat = EPixelFormat::PF_R16F;
+			}
+		}
+		return PixelFormat;
+	}
+
+	uint32 GetFeatureBytesPerElement()
+	{
+		uint32 ByteSize = sizeof(float);
+		if (GetRegressionDevice() != ERegressionDevice::CPU)
+		{
+			if (CVarNFORNOnLocalMeanFeatureFormat.GetValueOnRenderThread() != 0)
+			{
+				ByteSize = sizeof(int16_t);
+			}
+		}
+
+		return ByteSize;
+	}
+
 	EAlbedoDivideRecoverPhase GetPreAlbedoDivideRecoverPhase()
 	{
 		EAlbedoDivideRecoverPhase AlbedoDivideRecoverPhase = EAlbedoDivideRecoverPhase::Disabled;
@@ -1510,7 +1545,7 @@ namespace NFORDenoise
 			SHADER::FParameters* PassParameters = GraphBuilder.AllocParameters<SHADER::FParameters>();
 			{
 				PassParameters->Source = GraphBuilder.CreateSRV(FRDGTextureSRVDesc::Create(Source));
-				PassParameters->Dest = GraphBuilder.CreateUAV(Dest, PF_R32_FLOAT);
+				PassParameters->Dest = GraphBuilder.CreateUAV(Dest, SHADER::GetDestFloatFormat(Dest->Desc.BytesPerElement));
 				PassParameters->TextureSize = TextureSize;
 				PassParameters->CopyChannelOffset = CopyChannelOffset;
 				PassParameters->CopyChannelCount = CopyChannelCount;
@@ -1630,7 +1665,7 @@ namespace NFORDenoise
 
 			SHADER::FParameters* PassParameters = GraphBuilder.AllocParameters<SHADER::FParameters>();
 			{
-				PassParameters->X = GraphBuilder.CreateSRV(FRDGBufferSRVDesc(X, PF_R32_FLOAT));
+				PassParameters->X = GraphBuilder.CreateSRV(FRDGBufferSRVDesc(X, SHADER::GetXYFloatFormat(X->Desc.BytesPerElement)));
 				PassParameters->XDim = XDim;
 
 				PassParameters->W = GraphBuilder.CreateSRV(FRDGBufferSRVDesc(W, GetWeightLayoutPixelFormat(NonLocalMeanWeightLayout)));
@@ -1639,12 +1674,12 @@ namespace NFORDenoise
 
 				if (GeneralizedMultiplication)
 				{
-					PassParameters->Y = GraphBuilder.CreateSRV(FRDGBufferSRVDesc(Y, PF_R32_FLOAT));
+					PassParameters->Y = GraphBuilder.CreateSRV(FRDGBufferSRVDesc(Y, SHADER::GetXYFloatFormat(Y->Desc.BytesPerElement)));
 					PassParameters->YDim = YDim;
 				}
 				else
 				{
-					PassParameters->Y = GraphBuilder.CreateSRV(FRDGBufferSRVDesc(X, PF_R32_FLOAT));
+					PassParameters->Y = GraphBuilder.CreateSRV(FRDGBufferSRVDesc(X, SHADER::GetXYFloatFormat(X->Desc.BytesPerElement)));
 					PassParameters->YDim = XDim;
 				}
 
@@ -1737,7 +1772,7 @@ namespace NFORDenoise
 
 			SHADER::FParameters* PassParameters = GraphBuilder.AllocParameters<SHADER::FParameters>();
 			{
-				PassParameters->X = GraphBuilder.CreateSRV(FRDGBufferSRVDesc(Feature, PF_R32_FLOAT));
+				PassParameters->X = GraphBuilder.CreateSRV(FRDGBufferSRVDesc(Feature, SHADER::GetXFloatFormat(Feature->Desc.BytesPerElement)));
 				PassParameters->W = GraphBuilder.CreateSRV(FRDGBufferSRVDesc(NonLocalMeanWeightsBuffer, GetWeightLayoutPixelFormat(NonLocalMeanWeightLayout)));
 				PassParameters->B = GraphBuilder.CreateSRV(FRDGBufferSRVDesc(ReconstructionWeights, PF_R32_FLOAT));
 				PassParameters->RWReconstruction = GraphBuilder.CreateUAV(FilteredRadiance);
@@ -2312,7 +2347,7 @@ namespace NFORDenoise
 			NonLocalMeanWeightsBuffer = NonLocalMeanSingleFrameWeightsBuffer;
 		}
 
-		FRDGBufferDesc CombinedFeatureDesc = FRDGBufferDesc::CreateBufferDesc(BytesPerElement, PaddedTileSize.X * PaddedTileSize.Y * NumOfCombinedFeatureChannels);
+		FRDGBufferDesc CombinedFeatureDesc = FRDGBufferDesc::CreateBufferDesc(GetFeatureBytesPerElement(), PaddedTileSize.X * PaddedTileSize.Y * NumOfCombinedFeatureChannels);
 		FRDGBufferRef CombinedFeatures = GraphBuilder.CreateBuffer(CombinedFeatureDesc, TEXT("NFOR.CombinedFeatures"));
 
 		const int32 NumOfCombinedRadianceChannels = GetNumOfCombinedFeatureChannels(Radiances);
diff --git a/Engine/Plugins/Experimental/NFORDenoise/Source/NFORDenoise/Private/NFORDenoiseCS.h b/Engine/Plugins/Experimental/NFORDenoise/Source/NFORDenoise/Private/NFORDenoiseCS.h
--- a/Engine/Plugins/Experimental/NFORDenoise/Source/NFORDenoise/Private/NFORDenoiseCS.h
+++ b/Engine/Plugins/Experimental/NFORDenoise/Source/NFORDenoise/Private/NFORDenoiseCS.h
@@ -538,6 +538,17 @@ namespace NFORDenoise
 			FGlobalShader::ModifyCompilationEnvironment(InParameters, OutEnvironment);
 			OutEnvironment.SetDefine(TEXT("THREAD_GROUP_SIZE"), NON_LOCAL_MEAN_THREAD_GROUP_SIZE);
 		}
+		
+		static EPixelFormat GetDestFloatFormat(uint32 BufferBytesPerElement)
+		{
+			// Dest buffer can be float32 or float16
+			EPixelFormat BufferFloatFormat = PF_R32_FLOAT;
+			if (BufferBytesPerElement == sizeof(int16_t))
+			{
+				BufferFloatFormat = PF_R16F;
+			}
+			return BufferFloatFormat;
+		}
 
 		static const int kMaxSourceChannelCount = 4;
 
@@ -648,6 +659,17 @@ namespace NFORDenoise
 				OutEnvironment.SetDefine(TEXT("THREAD_GROUP_SIZE"), GetThreadGroupSize());
 			}
 
+			static EPixelFormat GetXYFloatFormat(uint32 BufferBytesPerElement)
+			{
+				// Buffer can be float32 or float16
+				EPixelFormat BufferFloatFormat = PF_R32_FLOAT;
+				if (BufferBytesPerElement == sizeof(int16_t))
+				{
+					BufferFloatFormat = PF_R16F;
+				}
+				return BufferFloatFormat;
+			}
+
 			static const int kMaxSourceChannelCount = 4;
 
 			class FDimWeightedMultiplicationType : SHADER_PERMUTATION_ENUM_CLASS("WEIGHTED_MULTIPLICATION_TYPE", EWeightedMultiplicationType);
@@ -799,6 +821,17 @@ namespace NFORDenoise
 				OutEnvironment.CompilerFlags.Add(CFLAG_AllowTypedUAVLoads);
 			}
 
+			static EPixelFormat GetXFloatFormat(uint32 BufferBytesPerElement)
+			{
+				// Buffer can be float32 or float16
+				EPixelFormat BufferFloatFormat = PF_R32_FLOAT;
+				if (BufferBytesPerElement == sizeof(int16_t))
+				{
+					BufferFloatFormat = PF_R16F;
+				}
+				return BufferFloatFormat;
+			}
+
 			enum class EReconstructionType : uint8
 			{
 				Scatter,
diff --git a/Engine/Source/Runtime/Renderer/Private/PathTracingSpatialTemporalDenoising.cpp b/Engine/Source/Runtime/Renderer/Private/PathTracingSpatialTemporalDenoising.cpp
--- a/Engine/Source/Runtime/Renderer/Private/PathTracingSpatialTemporalDenoising.cpp
+++ b/Engine/Source/Runtime/Renderer/Private/PathTracingSpatialTemporalDenoising.cpp
@@ -933,6 +933,7 @@ static void PathTracingSpatialTemporalDenoiserPlugin(FRDGBuilder& GraphBuilder,
 	FRDGTextureRef InputTexture,
 	FRDGTextureRef AlbedoTexture,
 	FRDGTextureRef NormalTexture,
+	FRDGTextureRef DepthTexture,
 	FRDGTextureRef FlowTexture,
 	FRDGTextureRef PreviousOutputFrameTexture,
 	FRDGTextureRef OutputTexture,
@@ -966,6 +967,7 @@ static void PathTracingSpatialTemporalDenoiserPlugin(FRDGBuilder& GraphBuilder,
 	Inputs.ColorTex = InputTexture;
 	Inputs.AlbedoTex = AlbedoTexture;
 	Inputs.NormalTex = NormalTexture;
+	Inputs.DepthTex = DepthTexture;
 	Inputs.VarianceTex = Context.VarianceTexture;
 	Inputs.OutputTex = OutputTexture;
 	Inputs.FlowTex = FlowTexture;
@@ -1908,6 +1910,7 @@ void PathTracingSpatialTemporalDenoising(FRDGBuilder& GraphBuilder,
 		PathTracingSpatialTemporalDenoiserPlugin(
 			GraphBuilder, View, DenoiserMode, TargetTexture,
 			SpatialTemporalDenoisingContext.AlbedoTexture, SpatialTemporalDenoisingContext.NormalTexture,
+			SpatialTemporalDenoisingContext.DepthTexture,
 			MotionTexture, SourceTexture, TemporalDenoisedTexture,
 			SpatialTemporalDenoisingContext.FrameIndex,
 			bIsInitialFrame,
diff --git a/Engine/Source/Runtime/Renderer/Public/PathTracingDenoiser.h b/Engine/Source/Runtime/Renderer/Public/PathTracingDenoiser.h
--- a/Engine/Source/Runtime/Renderer/Public/PathTracingDenoiser.h
+++ b/Engine/Source/Runtime/Renderer/Public/PathTracingDenoiser.h
@@ -61,6 +61,7 @@ public:
 		FRDGTextureRef ColorTex;
 		FRDGTextureRef AlbedoTex;
 		FRDGTextureRef NormalTex;
+		FRDGTextureRef DepthTex;
 		FRDGTextureRef VarianceTex;
 		FRDGTextureRef OutputTex;
 
