Change NFOR denoiser to use 5 frame denosing for offline by default (r.NFOR.FrameCount 2). It denoises single frame for online editor denoising preview (r.NFOR.FrameCount 0).
#rnx
#tests Hillside projects renders correctly in offline and online mode

diff --git a/Engine/Plugins/Experimental/NFORDenoise/Source/NFORDenoise/Private/NFORDenoise.cpp b/Engine/Plugins/Experimental/NFORDenoise/Source/NFORDenoise/Private/NFORDenoise.cpp
--- a/Engine/Plugins/Experimental/NFORDenoise/Source/NFORDenoise/Private/NFORDenoise.cpp
+++ b/Engine/Plugins/Experimental/NFORDenoise/Source/NFORDenoise/Private/NFORDenoise.cpp
@@ -122,7 +122,7 @@ public:
 
 		void AddFrame(FRDGBuilder& GraphBuilder, const FSceneView& View, const FInputs& Inputs)
 		{
-			const int NumFrames = NFORDenoise::GetFrameCount();
+			const int NumFrames = NFORDenoise::GetFrameCount(View);
 			if (NumFrames <= 0)
 			{
 				return;
@@ -225,11 +225,11 @@ public:
 
 		// Post process DOF replies on the depth buffer to estimate the DOF,
 		// Update the output depth buffer to the depth of the center image if the denoising frame is not the latest.
-		const bool bNeedToUpdateDepth = bDenoised && (NFORDenoise::GetFrameCount() > 1) && NFORDenoise::GetDenoisingFrameIndex(NumberOfHistory) > 0;
+		const bool bNeedToUpdateDepth = bDenoised && (NFORDenoise::GetFrameCount(View) > 1) && NFORDenoise::GetDenoisingFrameIndex(View, NumberOfHistory) > 0;
 		if (bNeedToUpdateDepth)
 		{
 				
-			int32 DenoisingDepthFrameIndex = NumberOfHistory - NFORDenoise::GetDenoisingFrameIndex(NumberOfHistory) - 1;
+			int32 DenoisingDepthFrameIndex = NumberOfHistory - NFORDenoise::GetDenoisingFrameIndex(View,NumberOfHistory) - 1;
 			if (DenoisingDepthFrameIndex >= 0)
 			{
 				const bool bAlphaOnly = true;
diff --git a/Engine/Plugins/Experimental/NFORDenoise/Source/NFORDenoise/Private/NFORDenoiseCS.cpp b/Engine/Plugins/Experimental/NFORDenoise/Source/NFORDenoise/Private/NFORDenoiseCS.cpp
--- a/Engine/Plugins/Experimental/NFORDenoise/Source/NFORDenoise/Private/NFORDenoiseCS.cpp
+++ b/Engine/Plugins/Experimental/NFORDenoise/Source/NFORDenoise/Private/NFORDenoiseCS.cpp
@@ -57,8 +57,9 @@ namespace NFORDenoise
 
 	TAutoConsoleVariable<int32> CVarNFORFrameCount(
 		TEXT("r.NFOR.FrameCount"),
-		0,
-		TEXT("n: Use the previous n frames, the current frame, and the future n frames. Suggested range is 0~2. Max=3."),
+		2,
+		TEXT("n: Use the previous n frames, the current frame, and the future n frames. Suggested range is 0~2. Max=3.(Offline config)\n")
+		TEXT("The value is always 0 for online preview denoising\n"),
 		ECVF_RenderThreadSafe);
 
 	TAutoConsoleVariable<int32> CVarNFORFrameCountCondition(
@@ -290,9 +291,13 @@ namespace NFORDenoise
 		return FLinearColor(Offset, OffsetSky, 0.0f);
 	}
 
-	int32 GetFrameCount()
+	int32 GetFrameCount(const FSceneView& View)
 	{
-		const int32 NumFrames = FMath::Clamp(1 + 2 * CVarNFORFrameCount.GetValueOnRenderThread(), 1, 7);
+		int32 NumFrames = FMath::Clamp(1 + 2 * CVarNFORFrameCount.GetValueOnRenderThread(), 1, 7);
+		if (!View.bIsOfflineRender)
+		{
+			NumFrames = 1;
+		}
 		return NumFrames;
 	}
 
@@ -314,9 +319,9 @@ namespace NFORDenoise
 		return Condition;
 	}
 
-	int32 GetDenoisingFrameIndex(int32 NumberOfFrameInBuffer)
+	int32 GetDenoisingFrameIndex(const FSceneView& View, int32 NumberOfFrameInBuffer)
 	{
-		int32 TargetFrameCount = GetFrameCount();
+		int32 TargetFrameCount = GetFrameCount(View);
 		int32 DenoisingFrameIndex = CVarNFORDenoisingFrameIndex.GetValueOnRenderThread();
 		int32 ResolvedSourceFrameIndex = INDEX_NONE;
 		if (DenoisingFrameIndex < 0)
@@ -1641,6 +1646,7 @@ namespace NFORDenoise
 
 	FRDGBufferRef ApplyBatchedInPlaceMatrixMultiplication(
 		FRDGBuilder& GraphBuilder,
+		const FSceneView& View,
 		FRDGBufferRef X,
 		FIntPoint XDim,
 		FRDGBufferRef W,
@@ -1687,7 +1693,7 @@ namespace NFORDenoise
 				PassParameters->PatchDistance = PatchDistance;
 				PassParameters->NumOfWeigthsPerPixelPerFrame = (PatchDistance * 2 + 1) * (PatchDistance * 2 + 1);
 				PassParameters->NumOfTemporalFrames = WDim / PassParameters->NumOfWeigthsPerPixelPerFrame;
-				PassParameters->SourceFrameIndex = GetDenoisingFrameIndex(PassParameters->NumOfTemporalFrames);
+				PassParameters->SourceFrameIndex = GetDenoisingFrameIndex(View, PassParameters->NumOfTemporalFrames);
 				PassParameters->SamplingStep = SamplingStep;
 				PassParameters->Result = GraphBuilder.CreateUAV(FRDGBufferUAVDesc(ResultMatrix, PF_R32_FLOAT));
 			}
@@ -2189,6 +2195,7 @@ namespace NFORDenoise
 		{
 			AMatrix = ApplyBatchedInPlaceMatrixMultiplication(
 				GraphBuilder,
+				View,
 				Feature,
 				XDimension,
 				NonLocalMeanWeightsBuffer,
@@ -2199,6 +2206,7 @@ namespace NFORDenoise
 
 			BMatrix = ApplyBatchedInPlaceMatrixMultiplication(
 				GraphBuilder,
+				View,
 				Feature,
 				XDimension,
 				NonLocalMeanWeightsBuffer,
@@ -2263,7 +2271,7 @@ namespace NFORDenoise
 				}
 
 				RegressionKernel::FReconstructSpatialTemporalImage::EReconstructionType
-					ReconstructionType = GetReconstructionType(FrameIndex, GetDenoisingFrameIndex(WeightedLSRDesc.NumOfFrames));
+					ReconstructionType = GetReconstructionType(FrameIndex, GetDenoisingFrameIndex(View,WeightedLSRDesc.NumOfFrames));
 				
 				ReconstructByFrame(
 					GraphBuilder,
@@ -2312,7 +2320,7 @@ namespace NFORDenoise
 	{
 		const int32 NumOfFeatures = Features.Num();
 		const int32 NumOfRadiances = Radiances.Num();
-		const int32 SourceIndex = GetDenoisingFrameIndex(NumOfRadiances); // The current denoising frame data index
+		const int32 SourceIndex = GetDenoisingFrameIndex(View,NumOfRadiances); // The current denoising frame data index
 
 		const int SearchingPatchSize = (RadianceNonLocalMeanParameters.PatchDistance * 2 + 1);
 		const int NumberOfWeightsPerPixel = SearchingPatchSize * SearchingPatchSize;
@@ -2633,7 +2641,7 @@ namespace NFORDenoise
 		const int32 NumOfTemporalFrames = Radiances.Num();
 		const int32 NumOfFeatures = FeatureDescs.Num();
 		const int32 NumOfFeaturesPerFrame = NumOfFeatures / NumOfTemporalFrames;
-		const int32 SourceRadianceIndex = GetDenoisingFrameIndex(NumOfTemporalFrames);
+		const int32 SourceRadianceIndex = GetDenoisingFrameIndex(View,NumOfTemporalFrames);
 		
 		// Preprocessing
 		// Feature range adjustment, radiance normalization and filtering frames
diff --git a/Engine/Plugins/Experimental/NFORDenoise/Source/NFORDenoise/Private/NFORDenoiseCS.h b/Engine/Plugins/Experimental/NFORDenoise/Source/NFORDenoise/Private/NFORDenoiseCS.h
--- a/Engine/Plugins/Experimental/NFORDenoise/Source/NFORDenoise/Private/NFORDenoiseCS.h
+++ b/Engine/Plugins/Experimental/NFORDenoise/Source/NFORDenoise/Private/NFORDenoiseCS.h
@@ -18,8 +18,8 @@
 namespace NFORDenoise
 {
 
-	int32 GetFrameCount();
-	int32 GetDenoisingFrameIndex(int32 NumberOfFrameInBuffer = -1);
+	int32 GetFrameCount(const FSceneView& View);
+	int32 GetDenoisingFrameIndex(const FSceneView& View, int32 NumberOfFrameInBuffer = -1);
 
 	enum class EVarianceType : uint32
 	{
