Fix NFOR denoise creates artifact when the scene feature is too sparse. The cholesky decomposition is ill-conditioned. Apply adaptive lambda regularization by the max of the absolution of matrix value.

#rnx
#tests Denoise as expected on simple scene and hillside project

diff --git a/Engine/Plugins/Experimental/NFORDenoise/Source/NFORDenoise/Private/NFORDenoiseCS.cpp b/Engine/Plugins/Experimental/NFORDenoise/Source/NFORDenoise/Private/NFORDenoiseCS.cpp
--- a/Engine/Plugins/Experimental/NFORDenoise/Source/NFORDenoise/Private/NFORDenoiseCS.cpp
+++ b/Engine/Plugins/Experimental/NFORDenoise/Source/NFORDenoise/Private/NFORDenoiseCS.cpp
@@ -143,8 +143,8 @@ namespace NFORDenoise
 
 	TAutoConsoleVariable<float> CVarNFORLinearSolverCholeskyLambda(
 		TEXT("r.NFOR.LinearSolver.Cholesky.Lambda"),
-		1e-3,
-		TEXT("The parameter lambda for modified Cholesky decomposition to make it positive definite.\n")
+		2e-5,
+		TEXT("The initial lambda for modified Cholesky decomposition to make it positive definite. It will be scaled by the max of the absolute of the matrix element.\n")
 		TEXT("Large value yields bias with smoothed rendering, while small value leads to variance or artifacts.\n")
 		TEXT("Used when r.NFOR.LinearSolver.Type = 1 and 2. Selected to match the quality of r.NFOR.LinearSolver.Type = 2\n")
 		TEXT("to r.NFOR.LinearSolver.Type 0."),
@@ -1908,6 +1908,7 @@ namespace NFORDenoise
 			CommonPassParameters.NumOfElements = NumOfElements;
 			CommonPassParameters.NumOfElementsPerRow = NumOfElementsPerRow;
 			CommonPassParameters.Lambda = 0.0f;
+			CommonPassParameters.MinLambda = 0.0f;
 		}
 
 		// First multi-pass or the single pass based on SolverType.
@@ -1919,11 +1920,13 @@ namespace NFORDenoise
 				*PassParameters = CommonPassParameters;
 				//Magnitude of X^TWX element value increases with the number of frames, and the number of elements selected to 
 				// estimate the weights. GetLinearSolverCholeskyLambda() returns the lambda for a single frame.
-				PassParameters->Lambda = GetLinearSolverCholeskyLambda() * WeightedLSRDesc.NumOfFrames;
+				PassParameters->Lambda = GetLinearSolverCholeskyLambda();
+				PassParameters->MinLambda = 1e-3 * WeightedLSRDesc.NumOfFrames; // Experimental value.
 				if (bApproximateGroundTruthSolver)
 				{
 					PassParameters->RWSuccessAndFailIndexBuffer = GraphBuilder.CreateUAV(FRDGBufferUAVDesc(SuccessAndFailIndexBuffer[0], PF_R32_UINT));
 					PassParameters->Lambda = 0.0f;
+					PassParameters->MinLambda = 0.0f;
 				}
 			}
 
diff --git a/Engine/Plugins/Experimental/NFORDenoise/Source/NFORDenoise/Private/NFORDenoiseCS.h b/Engine/Plugins/Experimental/NFORDenoise/Source/NFORDenoise/Private/NFORDenoiseCS.h
--- a/Engine/Plugins/Experimental/NFORDenoise/Source/NFORDenoise/Private/NFORDenoiseCS.h
+++ b/Engine/Plugins/Experimental/NFORDenoise/Source/NFORDenoise/Private/NFORDenoiseCS.h
@@ -701,6 +701,7 @@ namespace NFORDenoise
 				SHADER_PARAMETER(int, NumOfElementsPerRow)
 
 				SHADER_PARAMETER(float, Lambda)
+				SHADER_PARAMETER(float, MinLambda)
 				SHADER_PARAMETER_RDG_BUFFER_UAV(RWBuffer<uint>, RWSuccessAndFailIndexBuffer)
 			END_SHADER_PARAMETER_STRUCT()
 
diff --git a/Engine/Plugins/Experimental/NFORDenoise/Source/NFORDenoise/Shaders/NFORRegression.ush b/Engine/Plugins/Experimental/NFORDenoise/Source/NFORDenoise/Shaders/NFORRegression.ush
--- a/Engine/Plugins/Experimental/NFORDenoise/Source/NFORDenoise/Shaders/NFORRegression.ush
+++ b/Engine/Plugins/Experimental/NFORDenoise/Source/NFORDenoise/Shaders/NFORRegression.ush
@@ -35,6 +35,8 @@ int2 BDim;
 int NumOfElements;
 int NumOfElementsPerRow;
 
+float MinLambda;
+
 struct FLMatrix
 {
 	FDFScalar M[LOWER_TRIANGLE_MATRIX_SIZE];
@@ -97,6 +99,16 @@ FMatrix  LMatrixMultiply(in FLMatrix A, in FLMatrix B)
 	return Ret;
 }
 
+float LMatrixMaxAbsValue(in FLMatrix A)
+{
+	float MaxAbsValue = 0;
+	for (int i = 0; i < SYMMETRIC_MATRIX_SIZE; ++i)
+	{
+		MaxAbsValue = max(MaxAbsValue, abs(DFDemote(A.M[i])));
+	}
+	return MaxAbsValue;
+}
+
 // Assume A is symmetric, and B is not. However, A*B is symmetric
 FLMatrix  ToLMatrixMultiply(in FLMatrix A, in FMatrix B)
 {
@@ -440,6 +452,14 @@ bool CholeskeyDecomposition(inout FLMatrix L, FLMatrix AMatrix, float lambda)
 	bool bComplete = true;
 	L = (FLMatrix)0;
 
+	// We need adaptive lambda otherwise, we cannot find a one size fit for all to make 
+	// this decomposition always successful.
+	// To achieve this, we use the max value of the absolute of the matrtix by the lambda
+	lambda *= LMatrixMaxAbsValue(AMatrix);
+	
+	// In case lambda becomes too small, we use the min lambda to provide sufficient regularization. 
+	lambda = max(lambda, MinLambda);
+	
 	for (int i = 0; i < MATRIX_DIM; ++i)
 	{
 		for (int j = 0; j <= i; ++j)
