Fix shader bug and reduce total number of NFOR shader compile for platforms that do not run path tracing.
#tests Preview all platforms in editor without crash or compiling NFOR specific shaders

diff --git a/Engine/Plugins/Experimental/NFORDenoise/Source/NFORDenoise/Private/NFORDenoise.cpp b/Engine/Plugins/Experimental/NFORDenoise/Source/NFORDenoise/Private/NFORDenoise.cpp
--- a/Engine/Plugins/Experimental/NFORDenoise/Source/NFORDenoise/Private/NFORDenoise.cpp
+++ b/Engine/Plugins/Experimental/NFORDenoise/Source/NFORDenoise/Private/NFORDenoise.cpp
@@ -254,7 +254,7 @@ private:
 
 void FNFORDenoiseModule::StartupModule()
 {
-	UE_LOG(LogNFORDenoise, Log, TEXT("NeuralPostProcessing starting up"));
+	UE_LOG(LogNFORDenoise, Log, TEXT("NFORDenoise function starting up"));
 
 	const TSharedPtr<IPlugin> Plugin = IPluginManager::Get().FindPlugin(TEXT("NFORDenoise"));
 	if (Plugin.IsValid())
diff --git a/Engine/Plugins/Experimental/NFORDenoise/Source/NFORDenoise/Private/NFORDenoiseCS.cpp b/Engine/Plugins/Experimental/NFORDenoise/Source/NFORDenoise/Private/NFORDenoiseCS.cpp
--- a/Engine/Plugins/Experimental/NFORDenoise/Source/NFORDenoise/Private/NFORDenoiseCS.cpp
+++ b/Engine/Plugins/Experimental/NFORDenoise/Source/NFORDenoise/Private/NFORDenoiseCS.cpp
@@ -258,6 +258,16 @@ namespace NFORDenoise
 	//--------------------------------------------------------------------------------------------------------------------
 	// Functions based on CVars
 
+	bool ShouldCompileNFORShadersForProject(EShaderPlatform ShaderPlatform)
+	{
+		const IConsoleVariable* CVarPathTracing = IConsoleManager::Get().FindConsoleVariable(TEXT("r.PathTracing"));
+		const bool bSupportsPathTracing = CVarPathTracing ? CVarPathTracing->GetInt() != 0 : false;
+
+		return ShouldCompileRayTracingShadersForProject(ShaderPlatform) &&
+			FDataDrivenShaderPlatformInfo::GetSupportsPathTracing(ShaderPlatform) &&
+			bSupportsPathTracing;
+	}
+
 	bool ShouldFeatureAddConstant()
 	{
 		return CVarNFORFeatureAddConstant.GetValueOnRenderThread();
diff --git a/Engine/Plugins/Experimental/NFORDenoise/Source/NFORDenoise/Private/NFORDenoiseCS.h b/Engine/Plugins/Experimental/NFORDenoise/Source/NFORDenoise/Private/NFORDenoiseCS.h
--- a/Engine/Plugins/Experimental/NFORDenoise/Source/NFORDenoise/Private/NFORDenoiseCS.h
+++ b/Engine/Plugins/Experimental/NFORDenoise/Source/NFORDenoise/Private/NFORDenoiseCS.h
@@ -57,6 +57,8 @@ namespace NFORDenoise
 	//--------------------------------------------------------------------------------------------------------------------
 	// General texture operations including: multiply, divide, accumulate, and copy.
 
+	bool ShouldCompileNFORShadersForProject(EShaderPlatform ShaderPlatform);
+
 	class FTextureMultiplyCS : public FGlobalShader
 	{
 		DECLARE_GLOBAL_SHADER(FTextureMultiplyCS)
@@ -78,6 +80,11 @@ namespace NFORDenoise
 			OutEnvironment.SetDefine(TEXT("TEXTURE_OPS"), TEXTURE_OPS_MULTIPLY);
 			OutEnvironment.CompilerFlags.Add(CFLAG_AllowTypedUAVLoads);
 		}
+
+		static bool ShouldCompilePermutation(const FGlobalShaderPermutationParameters& Parameters)
+		{
+			return ShouldCompileNFORShadersForProject(Parameters.Platform);
+		}
 	};
 
 	class FTextureDivideCS : public FGlobalShader
@@ -101,6 +108,11 @@ namespace NFORDenoise
 			OutEnvironment.SetDefine(TEXT("TEXTURE_OPS"), TEXTURE_OPS_DIVIDE);
 			OutEnvironment.CompilerFlags.Add(CFLAG_AllowTypedUAVLoads);
 		}
+
+		static bool ShouldCompilePermutation(const FGlobalShaderPermutationParameters& Parameters)
+		{
+			return ShouldCompileNFORShadersForProject(Parameters.Platform);
+		}
 	};
 
 	class FTextureAccumulateConstantCS : public FGlobalShader
@@ -124,6 +136,11 @@ namespace NFORDenoise
 			OutEnvironment.CompilerFlags.Add(CFLAG_AllowTypedUAVLoads);
 		}
 
+		static bool ShouldCompilePermutation(const FGlobalShaderPermutationParameters& Parameters)
+		{
+			return ShouldCompileNFORShadersForProject(Parameters.Platform);
+		}
+
 		class FDimensionAccumulateByMask : SHADER_PERMUTATION_BOOL("ACCUMULATE_BY_MASK");
 		using FPermutationDomain = TShaderPermutationDomain<FDimensionAccumulateByMask>;
 	};
@@ -149,6 +166,11 @@ namespace NFORDenoise
 			OutEnvironment.SetDefine(TEXT("TEXTURE_OPS"), TEXTURE_OPS_ACCUMULATE);
 			OutEnvironment.CompilerFlags.Add(CFLAG_AllowTypedUAVLoads);
 		}
+
+		static bool ShouldCompilePermutation(const FGlobalShaderPermutationParameters& Parameters)
+		{
+			return ShouldCompileNFORShadersForProject(Parameters.Platform);
+		}
 	};
 
 	class FCopyTexturePS : public FGlobalShader
@@ -162,6 +184,11 @@ namespace NFORDenoise
 			SHADER_PARAMETER(FIntPoint, TextureSize)
 			RENDER_TARGET_BINDING_SLOTS()
 		END_SHADER_PARAMETER_STRUCT()
+
+		static bool ShouldCompilePermutation(const FGlobalShaderPermutationParameters& Parameters)
+		{
+			return ShouldCompileNFORShadersForProject(Parameters.Platform);
+		}
 	};
 
 	enum class ETextureCopyType :uint32
@@ -192,6 +219,11 @@ namespace NFORDenoise
 			OutEnvironment.SetDefine(TEXT("THREAD_GROUP_SIZE"), NON_LOCAL_MEAN_THREAD_GROUP_SIZE);
 		}
 
+		static bool ShouldCompilePermutation(const FGlobalShaderPermutationParameters& Parameters)
+		{
+			return ShouldCompileNFORShadersForProject(Parameters.Platform);
+		}
+
 		class FDimTextureCopyType : SHADER_PERMUTATION_ENUM_CLASS("TEXTURE_COPY_TYPE", ETextureCopyType);
 		using FPermutationDomain = TShaderPermutationDomain<FDimTextureCopyType>;
 	};
@@ -241,6 +273,11 @@ namespace NFORDenoise
 			FGlobalShader::ModifyCompilationEnvironment(InParameters, OutEnvironment);
 			OutEnvironment.SetDefine(TEXT("THREAD_GROUP_SIZE"), NON_LOCAL_MEAN_THREAD_GROUP_SIZE);
 		}
+
+		static bool ShouldCompilePermutation(const FGlobalShaderPermutationParameters& Parameters)
+		{
+			return ShouldCompileNFORShadersForProject(Parameters.Platform);
+		}
 	};
 
 	// Based on Taylor expansion, sigma_{normalized radiance} \approx = sigma_{radiance} / sigma_{albedo}, if normalized radiance = radiance / albedo
@@ -261,6 +298,11 @@ namespace NFORDenoise
 			OutEnvironment.SetDefine(TEXT("THREAD_GROUP_SIZE"), NON_LOCAL_MEAN_THREAD_GROUP_SIZE);
 			OutEnvironment.CompilerFlags.Add(CFLAG_AllowTypedUAVLoads);
 		}
+
+		static bool ShouldCompilePermutation(const FGlobalShaderPermutationParameters& Parameters)
+		{
+			return ShouldCompileNFORShadersForProject(Parameters.Platform);
+		}
 	};
 
 	struct FFeatureDesc;
@@ -287,6 +329,11 @@ namespace NFORDenoise
 			OutEnvironment.CompilerFlags.Add(CFLAG_AllowTypedUAVLoads);
 		}
 
+		static bool ShouldCompilePermutation(const FGlobalShaderPermutationParameters& Parameters)
+		{
+			return ShouldCompileNFORShadersForProject(Parameters.Platform);
+		}
+
 		class FDimensionVarianceType : SHADER_PERMUTATION_ENUM_CLASS("IMAGE_VARIANCE_TYPE", EVarianceType);
 		using FPermutationDomain = TShaderPermutationDomain<FDimensionVarianceType>;
 	};
@@ -374,6 +421,11 @@ namespace NFORDenoise
 			OutEnvironment.SetDefine(TEXT("NONLOCALMEAN_SEPARATE_SOURCE"), 0);
 		}
 
+		static bool ShouldCompilePermutation(const FGlobalShaderPermutationParameters& Parameters)
+		{
+			return ShouldCompileNFORShadersForProject(Parameters.Platform);
+		}
+
 		class FDimensionVarianceType : SHADER_PERMUTATION_ENUM_CLASS("IMAGE_VARIANCE_TYPE", EVarianceType);
 		class FDimensionUseGuide : SHADER_PERMUTATION_BOOL("USE_GUIDE"); //TODO
 		class FDimensionImageChannelCount : SHADER_PERMUTATION_RANGE_INT("SOURCE_CHANNEL_COUNT", 1, static_cast<int>(EImageChannelCount::MAX));
@@ -408,6 +460,11 @@ namespace NFORDenoise
 			OutEnvironment.SetDefine(TEXT("THREAD_GROUP_SIZE"), NON_LOCAL_MEAN_THREAD_GROUP_SIZE);
 		}
 
+		static bool ShouldCompilePermutation(const FGlobalShaderPermutationParameters& Parameters)
+		{
+			return ShouldCompileNFORShadersForProject(Parameters.Platform);
+		}
+
 		class FDimensionVarianceType : SHADER_PERMUTATION_ENUM_CLASS("IMAGE_VARIANCE_TYPE", EVarianceType);
 		class FDimensionUseGuide : SHADER_PERMUTATION_BOOL("USE_GUIDE");
 		class FDimensionImageChannelCount : SHADER_PERMUTATION_RANGE_INT("SOURCE_CHANNEL_COUNT", 1, static_cast<int>(EImageChannelCount::MAX));
@@ -439,6 +496,11 @@ namespace NFORDenoise
 			OutEnvironment.SetDefine(TEXT("THREAD_GROUP_SIZE"), NON_LOCAL_MEAN_THREAD_GROUP_SIZE);
 		}
 
+		static bool ShouldCompilePermutation(const FGlobalShaderPermutationParameters& Parameters)
+		{
+			return ShouldCompileNFORShadersForProject(Parameters.Platform);
+		}
+
 		class FDimensionVarianceType : SHADER_PERMUTATION_ENUM_CLASS("IMAGE_VARIANCE_TYPE", EVarianceType);
 		class FDimensionImageChannelCount : SHADER_PERMUTATION_RANGE_INT("SOURCE_CHANNEL_COUNT", 1, static_cast<int>(EImageChannelCount::MAX));
 		class FDimensionSeparateSourceTarget : SHADER_PERMUTATION_BOOL("NONLOCALMEAN_SEPARATE_SOURCE");
@@ -467,6 +529,11 @@ namespace NFORDenoise
 			OutEnvironment.SetDefine(TEXT("THREAD_GROUP_SIZE"), NON_LOCAL_MEAN_THREAD_GROUP_SIZE);
 		}
 
+		static bool ShouldCompilePermutation(const FGlobalShaderPermutationParameters& Parameters)
+		{
+			return ShouldCompileNFORShadersForProject(Parameters.Platform);
+		}
+
 		enum class ESeperablePassType : int32
 		{
 			Horizontal,
@@ -503,6 +570,11 @@ namespace NFORDenoise
 			OutEnvironment.SetDefine(TEXT("THREAD_GROUP_SIZE"), NON_LOCAL_MEAN_THREAD_GROUP_SIZE);
 		}
 
+		static bool ShouldCompilePermutation(const FGlobalShaderPermutationParameters& Parameters)
+		{
+			return ShouldCompileNFORShadersForProject(Parameters.Platform);
+		}
+
 		class FDimensionSeparateSourceTarget : SHADER_PERMUTATION_BOOL("NONLOCALMEAN_SEPARATE_SOURCE");
 		class FDimensionTargetWeightLayout : SHADER_PERMUTATION_ENUM_CLASS("NLM_WEIGHTLAYOUT", ENonLocalMeanWeightLayout);
 		using FPermutationDomain = TShaderPermutationDomain<FDimensionSeparateSourceTarget, FDimensionTargetWeightLayout>;
@@ -539,6 +611,11 @@ namespace NFORDenoise
 			OutEnvironment.SetDefine(TEXT("THREAD_GROUP_SIZE"), NON_LOCAL_MEAN_THREAD_GROUP_SIZE);
 		}
 		
+		static bool ShouldCompilePermutation(const FGlobalShaderPermutationParameters& Parameters)
+		{
+			return ShouldCompileNFORShadersForProject(Parameters.Platform);
+		}
+
 		static EPixelFormat GetDestFloatFormat(uint32 BufferBytesPerElement)
 		{
 			// Dest buffer can be float32 or float16
@@ -572,6 +649,11 @@ namespace NFORDenoise
 			OutEnvironment.SetDefine(TEXT("THREAD_GROUP_SIZE"), NON_LOCAL_MEAN_THREAD_GROUP_SIZE);
 			OutEnvironment.CompilerFlags.Add(CFLAG_AllowTypedUAVLoads);
 		}
+
+		static bool ShouldCompilePermutation(const FGlobalShaderPermutationParameters& Parameters)
+		{
+			return ShouldCompileNFORShadersForProject(Parameters.Platform);
+		}
 	};
 
 
@@ -659,6 +741,11 @@ namespace NFORDenoise
 				OutEnvironment.SetDefine(TEXT("THREAD_GROUP_SIZE"), GetThreadGroupSize());
 			}
 
+			static bool ShouldCompilePermutation(const FGlobalShaderPermutationParameters& Parameters)
+			{
+				return ShouldCompileNFORShadersForProject(Parameters.Platform);
+			}
+
 			static EPixelFormat GetXYFloatFormat(uint32 BufferBytesPerElement)
 			{
 				// Buffer can be float32 or float16
@@ -719,6 +806,11 @@ namespace NFORDenoise
 				OutEnvironment.SetDefine(TEXT("THREAD_GROUP_SIZE"), NON_LOCAL_MEAN_THREAD_GROUP_SIZE);
 			}
 
+			static bool ShouldCompilePermutation(const FGlobalShaderPermutationParameters& Parameters)
+			{
+				return ShouldCompileNFORShadersForProject(Parameters.Platform);
+			}
+
 			class FDimNumFeature : SHADER_PERMUTATION_RANGE_INT("NUM_FEATURE", 6, 3);
 			class FDimSolverType : SHADER_PERMUTATION_ENUM_CLASS("LINEAR_SOLVER_TYPE", ESolverType);
 			class FDimOutputIndices : SHADER_PERMUTATION_BOOL("OUTPUT_INDICES");
@@ -749,6 +841,11 @@ namespace NFORDenoise
 				OutEnvironment.SetDefine(TEXT("THREAD_GROUP_SIZE"), NON_LOCAL_MEAN_THREAD_GROUP_SIZE);
 			}
 
+			static bool ShouldCompilePermutation(const FGlobalShaderPermutationParameters& Parameters)
+			{
+				return ShouldCompileNFORShadersForProject(Parameters.Platform);
+			}
+
 			class FDimInputMatrixType : SHADER_PERMUTATION_ENUM_CLASS("INPUT_MATRIX_TYPE", EInputMatrixType);
 			using FPermutationDomain = TShaderPermutationDomain <FDimInputMatrixType>;
 		};
@@ -771,6 +868,11 @@ namespace NFORDenoise
 				OutEnvironment.SetDefine(TEXT("THREAD_GROUP_SIZE"), NON_LOCAL_MEAN_THREAD_GROUP_SIZE);
 			}
 
+			static bool ShouldCompilePermutation(const FGlobalShaderPermutationParameters& Parameters)
+			{
+				return ShouldCompileNFORShadersForProject(Parameters.Platform);
+			}
+
 			class FDimNumFeature : SHADER_PERMUTATION_RANGE_INT("NUM_FEATURE", 6, 3);
 			class FDimSolverType : SHADER_PERMUTATION_ENUM_CLASS("LINEAR_SOLVER_TYPE", FLinearSolverCS::ESolverType);
 			class FDimInputMatrixType : SHADER_PERMUTATION_ENUM_CLASS("INPUT_MATRIX_TYPE", EInputMatrixType);
@@ -822,6 +924,11 @@ namespace NFORDenoise
 				OutEnvironment.CompilerFlags.Add(CFLAG_AllowTypedUAVLoads);
 			}
 
+			static bool ShouldCompilePermutation(const FGlobalShaderPermutationParameters& Parameters)
+			{
+				return ShouldCompileNFORShadersForProject(Parameters.Platform);
+			}
+
 			static EPixelFormat GetXFloatFormat(uint32 BufferBytesPerElement)
 			{
 				// Buffer can be float32 or float16
@@ -879,6 +986,11 @@ namespace NFORDenoise
 			OutEnvironment.CompilerFlags.Add(CFLAG_AllowTypedUAVLoads);
 		}
 
+		static bool ShouldCompilePermutation(const FGlobalShaderPermutationParameters& Parameters)
+		{
+			return ShouldCompileNFORShadersForProject(Parameters.Platform);
+		}
+
 		class FDimPreAlbedoDivide : SHADER_PERMUTATION_ENUM_CLASS("PRE_ALBEDO_DIVIDE", EAlbedoDivideRecoverPhase);
 		using FPermutationDomain = TShaderPermutationDomain<FDimPreAlbedoDivide>;
 	};
@@ -909,6 +1021,11 @@ namespace NFORDenoise
 			OutEnvironment.SetDefine(TEXT("SOURCE_CHANNEL_COUNT"), 4);
 		}
 
+		static bool ShouldCompilePermutation(const FGlobalShaderPermutationParameters& Parameters)
+		{
+			return ShouldCompileNFORShadersForProject(Parameters.Platform);
+		}
+
 		class FDimensionVarianceType : SHADER_PERMUTATION_ENUM_CLASS("IMAGE_VARIANCE_TYPE", EVarianceType);
 		using FPermutationDomain = TShaderPermutationDomain<FDimensionVarianceType>;
 	};
@@ -930,6 +1047,11 @@ namespace NFORDenoise
 			FGlobalShader::ModifyCompilationEnvironment(InParameters, OutEnvironment);
 			OutEnvironment.SetDefine(TEXT("THREAD_GROUP_SIZE"), NON_LOCAL_MEAN_THREAD_GROUP_SIZE);
 		}
+
+		static bool ShouldCompilePermutation(const FGlobalShaderPermutationParameters& Parameters)
+		{
+			return ShouldCompileNFORShadersForProject(Parameters.Platform);
+		}
 	};
 
 	class FCombineFilteredImageCS : public FGlobalShader
@@ -950,6 +1072,11 @@ namespace NFORDenoise
 			FGlobalShader::ModifyCompilationEnvironment(InParameters, OutEnvironment);
 			OutEnvironment.SetDefine(TEXT("THREAD_GROUP_SIZE"), NON_LOCAL_MEAN_THREAD_GROUP_SIZE);
 		}
+
+		static bool ShouldCompilePermutation(const FGlobalShaderPermutationParameters& Parameters)
+		{
+			return ShouldCompileNFORShadersForProject(Parameters.Platform);
+		}
 	};
 
 	//--------------------------------------------------------------------------------------------------------------------
diff --git a/Engine/Plugins/Experimental/NFORDenoise/Source/NFORDenoise/Shaders/NFORDenoise.usf b/Engine/Plugins/Experimental/NFORDenoise/Source/NFORDenoise/Shaders/NFORDenoise.usf
--- a/Engine/Plugins/Experimental/NFORDenoise/Source/NFORDenoise/Shaders/NFORDenoise.usf
+++ b/Engine/Plugins/Experimental/NFORDenoise/Source/NFORDenoise/Shaders/NFORDenoise.usf
@@ -257,6 +257,7 @@ int2 SourceOffset;
 int2 TargetOffset;
 int Channel;
 int2 CopySize;
+int2 TextureSize;
 
 float4 CopyTexturePS(float4 SvPosition : SV_POSITION) : SV_Target0
 {
@@ -326,7 +327,6 @@ float GetImageVariance(int2 P, Texture2D inVariance, int inVarianceChannelOffset
 Texture2D<float4> Normal;
 Texture2D<float4> NormalVariance;
 RWTexture2D<uint> RWMask;
-int2 TextureSize;
 
 [numthreads(THREAD_GROUP_SIZE, THREAD_GROUP_SIZE, 1)]
 void ClassifyPreAlbedoDivideMaskIdCS(in const uint3 DispatchThreadID : SV_DispatchThreadID)
@@ -1456,6 +1456,7 @@ uint4 EncodeFloat4(float4 Value)
 float4 DecodeFloat4(uint4 Value)
 {
 	float4 Decoded = float4(Value>>16) + float4(Value & 0xffff)/65536;
+	return Decoded;
 }
 
 void SaveFloat4ToBufferOrTexture(uint2 Position, float4 Value, uint NumOfScattering)
@@ -1482,7 +1483,7 @@ void SaveFloat4ToBufferOrTexture(uint2 Position, float4 Value, uint NumOfScatter
 
 #else
 	int IntermediateBufferIndex = Position.x + Position.y * (TextureSize.x + 2 * PatchDistance);
-	uint4 XBWEncoded = EncodeFloat4(XBW);
+	uint4 XBWEncoded = EncodeFloat4(Value);
 	float4 OldValue = 0;
 	InterlockedAdd(RWReconstructBuffer[IntermediateBufferIndex].x, XBWEncoded.x, OldValue.x);
 	InterlockedAdd(RWReconstructBuffer[IntermediateBufferIndex].y, XBWEncoded.y, OldValue.y);
