Cache FindConsoleObject?? for r.PathTracing in NFOR.

#
#rnx

diff --git a/Engine/Plugins/Experimental/NFORDenoise/Source/NFORDenoise/Private/NFORDenoiseCS.cpp b/Engine/Plugins/Experimental/NFORDenoise/Source/NFORDenoise/Private/NFORDenoiseCS.cpp
--- a/Engine/Plugins/Experimental/NFORDenoise/Source/NFORDenoise/Private/NFORDenoiseCS.cpp
+++ b/Engine/Plugins/Experimental/NFORDenoise/Source/NFORDenoise/Private/NFORDenoiseCS.cpp
@@ -260,7 +260,7 @@ namespace NFORDenoise
 
 	bool ShouldCompileNFORShadersForProject(EShaderPlatform ShaderPlatform)
 	{
-		const IConsoleVariable* CVarPathTracing = IConsoleManager::Get().FindConsoleVariable(TEXT("r.PathTracing"));
+		static const IConsoleVariable* CVarPathTracing = IConsoleManager::Get().FindConsoleVariable(TEXT("r.PathTracing"));
 		const bool bSupportsPathTracing = CVarPathTracing ? CVarPathTracing->GetInt() != 0 : false;
 
 		return ShouldCompileRayTracingShadersForProject(ShaderPlatform) &&
