Fix shader file path not found on some platforms.

diff --git a/Engine/Plugins/Experimental/NFORDenoise/Source/NFORDenoise/Private/NFORDenoise.cpp b/Engine/Plugins/Experimental/NFORDenoise/Source/NFORDenoise/Private/NFORDenoise.cpp
--- a/Engine/Plugins/Experimental/NFORDenoise/Source/NFORDenoise/Private/NFORDenoise.cpp
+++ b/Engine/Plugins/Experimental/NFORDenoise/Source/NFORDenoise/Private/NFORDenoise.cpp
@@ -260,9 +260,12 @@ void FNFORDenoiseModule::StartupModule()
 	if (Plugin.IsValid())
 	{
 		FString ModuleDir = Plugin->GetBaseDir() + TEXT("/Source/NFORDenoise");
-		AddShaderSourceDirectoryMapping(TEXT("/NFORDenoise"), FPaths::Combine(ModuleDir, TEXT("Shaders")));
-
-		RegisterSpatialTemporalDenoiser(MakeUnique<FNFORDenoiser>(),TEXT("NFOR"));
+		FString ModuleShaderDir = FPaths::Combine(ModuleDir, TEXT("Shaders"));
+		if (FPaths::DirectoryExists(ModuleShaderDir))
+		{
+			AddShaderSourceDirectoryMapping(TEXT("/NFORDenoise"), ModuleShaderDir);
+			RegisterSpatialTemporalDenoiser(MakeUnique<FNFORDenoiser>(), TEXT("NFOR"));
+		}
 	}
 	else
 	{
