Fix minor shader warnings in NFOR denoiser

#rnx

diff --git a/Engine/Plugins/Experimental/NFORDenoise/Source/NFORDenoise/Shaders/NFORDenoise.usf b/Engine/Plugins/Experimental/NFORDenoise/Source/NFORDenoise/Shaders/NFORDenoise.usf
--- a/Engine/Plugins/Experimental/NFORDenoise/Source/NFORDenoise/Shaders/NFORDenoise.usf
+++ b/Engine/Plugins/Experimental/NFORDenoise/Source/NFORDenoise/Shaders/NFORDenoise.usf
@@ -9,6 +9,10 @@
 #define SOURCE_CHANNEL_COUNT 4
 #endif
 
+#if !COMPUTESHADER
+#define THREAD_GROUP_SIZE 1
+#endif
+
 #if SOURCE_CHANNEL_COUNT == 1
 #define TPixelValue float
 #elif SOURCE_CHANNEL_COUNT == 2
@@ -259,7 +263,7 @@ int Channel;
 int2 CopySize;
 int2 TextureSize;
 
-float4 CopyTexturePS(float4 SvPosition : SV_POSITION) : SV_Target0
+TPixelValue CopyTexturePS(float4 SvPosition : SV_POSITION) : SV_Target0
 {
 	uint2 Position = (uint2)SvPosition.xy + SourceOffset;
 	return Source.Load(uint3(GetMirroredPosition(Position, TextureSize), 0));
