Fix EngineTest crash due to macro not defined.

diff --git a/Engine/Plugins/Experimental/NFORDenoise/Source/NFORDenoise/Shaders/NFORRegression.ush b/Engine/Plugins/Experimental/NFORDenoise/Source/NFORDenoise/Shaders/NFORRegression.ush
--- a/Engine/Plugins/Experimental/NFORDenoise/Source/NFORDenoise/Shaders/NFORRegression.ush
+++ b/Engine/Plugins/Experimental/NFORDenoise/Source/NFORDenoise/Shaders/NFORRegression.ush
@@ -18,6 +18,10 @@
 #define MATRIX_DIM 8
 #endif
 
+#ifndef NUM_NEWTON_SCHULTZ_ITERATIONS
+#define NUM_NEWTON_SCHULTZ_ITERATIONS 3
+#endif
+
 #define FULL_MATRIX_SIZE ((MATRIX_DIM)*(MATRIX_DIM))
 #define SYMMETRIC_MATRIX_SIZE (((MATRIX_DIM+1)*MATRIX_DIM)/2)
 #define LOWER_TRIANGLE_MATRIX_SIZE SYMMETRIC_MATRIX_SIZE
