Fix can't map virtual shader source path for NFOR.

diff --git a/Engine/Plugins/Experimental/NFORDenoise/Source/NFORDenoise/Private/NFORDenoise.cpp b/Engine/Plugins/Experimental/NFORDenoise/Source/NFORDenoise/Private/NFORDenoise.cpp
--- a/Engine/Plugins/Experimental/NFORDenoise/Source/NFORDenoise/Private/NFORDenoise.cpp
+++ b/Engine/Plugins/Experimental/NFORDenoise/Source/NFORDenoise/Private/NFORDenoise.cpp
@@ -261,7 +261,13 @@ void FNFORDenoiseModule::StartupModule()
 	{
 		FString ModuleDir = Plugin->GetBaseDir() + TEXT("/Source/NFORDenoise");
 		FString ModuleShaderDir = FPaths::Combine(ModuleDir, TEXT("Shaders"));
-		if (FPaths::DirectoryExists(ModuleShaderDir))
+
+		// If there is no editor data and does not allow shader compiling, we should not
+		// register the denoiser.
+		const bool bShouldDisableShaderCompiling = FPlatformProperties::RequiresCookedData() || !AllowShaderCompiling();
+
+		if (FPaths::DirectoryExists(ModuleShaderDir)
+			&& !bShouldDisableShaderCompiling)
 		{
 			AddShaderSourceDirectoryMapping(TEXT("/NFORDenoise"), ModuleShaderDir);
 			RegisterSpatialTemporalDenoiser(MakeUnique<FNFORDenoiser>(), TEXT("NFOR"));
