Fix NFOR denoising artifact when denoising value is too large. Swich denoising texture format from fp16 to fp32 during accumulation.

#tests Renders in RenderGallery without the overflow artifact.
#rnx

diff --git a/Engine/Plugins/Experimental/NFORDenoise/Shaders/NFORDenoise.usf b/Engine/Plugins/Experimental/NFORDenoise/Shaders/NFORDenoise.usf
--- a/Engine/Plugins/Experimental/NFORDenoise/Shaders/NFORDenoise.usf
+++ b/Engine/Plugins/Experimental/NFORDenoise/Shaders/NFORDenoise.usf
@@ -1465,8 +1465,9 @@ float4 DecodeFloat4(uint4 Value)
 
 void SaveFloat4ToBufferOrTexture(uint2 Position, float4 Value, uint NumOfScattering)
 {
+	Value = max(Value, 0);
 #if COMPILER_SUPPORTS_UINT64_IMAGE_ATOMICS
-	float4 XBWFraction = frac(max(Value, 0));
+	float4 XBWFraction = frac(Value);
 	uint4 XBWInteger = Value - XBWFraction;
 	XBWInteger = min(XBWInteger, 0xffffffff / NumOfScattering); // Avoid overflow.
 
@@ -1635,7 +1636,7 @@ float4 DecodeFloat4FromBufferOrTexture(uint2 Position)
 
 	// There would be cases where a is 0 due to w being too small and encoded as 0.
 	// Without rejecting this sample, could see borders artifacts around each tile.
-	DecodedValue = lerp(0, DecodedValue, DecodedValue.a != 0);
+	DecodedValue = lerp(0, DecodedValue, DecodedValue.a > 0);
 
 	return DecodedValue;
 }
diff --git a/Engine/Plugins/Experimental/NFORDenoise/Source/NFORDenoise/Private/NFORDenoise.cpp b/Engine/Plugins/Experimental/NFORDenoise/Source/NFORDenoise/Private/NFORDenoise.cpp
--- a/Engine/Plugins/Experimental/NFORDenoise/Source/NFORDenoise/Private/NFORDenoise.cpp
+++ b/Engine/Plugins/Experimental/NFORDenoise/Source/NFORDenoise/Private/NFORDenoise.cpp
@@ -131,11 +131,11 @@ public:
 			RDG_EVENT_SCOPE(GraphBuilder, "AddFrame");
 
 			FRDGFrameCache RDGFrameCache;
-			RDGFrameCache.Radiance = GraphBuilder.CreateTexture(Inputs.ColorTex->Desc,TEXT("Radiance"));
-			RDGFrameCache.RadianceVariance = GraphBuilder.CreateTexture(Inputs.VarianceTex->Desc, TEXT("RadianceVariance"));
-			RDGFrameCache.Albedo = GraphBuilder.CreateTexture(Inputs.AlbedoTex->Desc, TEXT("Albedo"));
-			RDGFrameCache.Normal = GraphBuilder.CreateTexture(Inputs.NormalTex->Desc, TEXT("Normal"));
-			RDGFrameCache.Depth = GraphBuilder.CreateTexture(Inputs.DepthTex->Desc, TEXT("Depth"));
+			RDGFrameCache.Radiance = GraphBuilder.CreateTexture(Inputs.ColorTex->Desc,TEXT("NFOR.Radiance"), ERDGTextureFlags::MultiFrame);
+			RDGFrameCache.RadianceVariance = GraphBuilder.CreateTexture(Inputs.VarianceTex->Desc, TEXT("NFOR.RadianceVariance"), ERDGTextureFlags::MultiFrame);
+			RDGFrameCache.Albedo = GraphBuilder.CreateTexture(Inputs.AlbedoTex->Desc, TEXT("NFOR.Albedo"), ERDGTextureFlags::MultiFrame);
+			RDGFrameCache.Normal = GraphBuilder.CreateTexture(Inputs.NormalTex->Desc, TEXT("NFOR.Normal"), ERDGTextureFlags::MultiFrame);
+			RDGFrameCache.Depth = GraphBuilder.CreateTexture(Inputs.DepthTex->Desc, TEXT("NFOR.Depth"), ERDGTextureFlags::MultiFrame);
 
 			AddCopyTexturePass(GraphBuilder, Inputs.ColorTex, RDGFrameCache.Radiance);
 			AddCopyTexturePass(GraphBuilder, Inputs.VarianceTex, RDGFrameCache.RadianceVariance);
diff --git a/Engine/Plugins/Experimental/NFORDenoise/Source/NFORDenoise/Private/NFORDenoiseCS.cpp b/Engine/Plugins/Experimental/NFORDenoise/Source/NFORDenoise/Private/NFORDenoiseCS.cpp
--- a/Engine/Plugins/Experimental/NFORDenoise/Source/NFORDenoise/Private/NFORDenoiseCS.cpp
+++ b/Engine/Plugins/Experimental/NFORDenoise/Source/NFORDenoise/Private/NFORDenoiseCS.cpp
@@ -2401,7 +2401,12 @@ namespace NFORDenoise
 		FRDGBufferRef CombinedRadiances = GraphBuilder.CreateBuffer(CombinedRadianceDesc, TEXT("NFOR.CombinedRadiances"));
 
 		FRDGTextureDesc FilteredRadianceDesc = Radiances[SourceIndex].Data.Image->Desc;
-		FilteredRadianceDesc.Flags |= TexCreate_RenderTargetable;
+		{
+			FilteredRadianceDesc.Flags |= TexCreate_RenderTargetable;
+		
+			check(FilteredRadianceDesc.Format == EPixelFormat::PF_FloatRGBA);
+			FilteredRadianceDesc.Format = EPixelFormat::PF_A32B32G32R32F;// The accumulation can run more than 2^16.
+		}
 		FRDGTextureRef FilteredRadiance = GraphBuilder.CreateTexture(FilteredRadianceDesc, TEXT("NFOR.FilteredRadiance"));
 
 		FRDGTextureDesc RadianceTileDesc = FilteredRadianceDesc;
@@ -2560,6 +2565,15 @@ namespace NFORDenoise
 		// Normalize the image by weights stored in alpha channel.
 		AddNormalizeTexturePass(GraphBuilder, FilteredRadiance);
 
+		// Copy back with the original format.
+		{
+			FRDGTextureDesc FilteredRadianceOutputDesc = Radiances[SourceIndex].Data.Image->Desc;
+
+			FRDGTextureRef FilteredRadianceOutputTexture = GraphBuilder.CreateTexture(FilteredRadianceOutputDesc, TEXT("NFOR.FilteredRadiance.Output"));
+			AddCopyMirroredTexturePass(GraphBuilder, FilteredRadiance, FilteredRadianceOutputTexture);
+			FilteredRadiance = FilteredRadianceOutputTexture;
+		}
+
 		return FilteredRadiance;
 	}
 
