Add format type check for NFOR internal denoising texture buffer.

diff --git a/Engine/Plugins/Experimental/NFORDenoise/Source/NFORDenoise/Private/NFORDenoiseCS.cpp b/Engine/Plugins/Experimental/NFORDenoise/Source/NFORDenoise/Private/NFORDenoiseCS.cpp
--- a/Engine/Plugins/Experimental/NFORDenoise/Source/NFORDenoise/Private/NFORDenoiseCS.cpp
+++ b/Engine/Plugins/Experimental/NFORDenoise/Source/NFORDenoise/Private/NFORDenoiseCS.cpp
@@ -2404,8 +2404,10 @@ namespace NFORDenoise
 		{
 			FilteredRadianceDesc.Flags |= TexCreate_RenderTargetable;
 		
-			check(FilteredRadianceDesc.Format == EPixelFormat::PF_FloatRGBA);
-			FilteredRadianceDesc.Format = EPixelFormat::PF_A32B32G32R32F;// The accumulation can run more than 2^16.
+			if(FilteredRadianceDesc.Format == EPixelFormat::PF_FloatRGBA)
+			{
+				FilteredRadianceDesc.Format = EPixelFormat::PF_A32B32G32R32F;// The accumulation can run more than 2^16.
+			}
 		}
 		FRDGTextureRef FilteredRadiance = GraphBuilder.CreateTexture(FilteredRadianceDesc, TEXT("NFOR.FilteredRadiance"));
 
@@ -2568,10 +2570,12 @@ namespace NFORDenoise
 		// Copy back with the original format.
 		{
 			FRDGTextureDesc FilteredRadianceOutputDesc = Radiances[SourceIndex].Data.Image->Desc;
-
-			FRDGTextureRef FilteredRadianceOutputTexture = GraphBuilder.CreateTexture(FilteredRadianceOutputDesc, TEXT("NFOR.FilteredRadiance.Output"));
-			AddCopyMirroredTexturePass(GraphBuilder, FilteredRadiance, FilteredRadianceOutputTexture);
-			FilteredRadiance = FilteredRadianceOutputTexture;
+			if (FilteredRadianceOutputDesc.Format == EPixelFormat::PF_FloatRGBA)
+			{
+				FRDGTextureRef FilteredRadianceOutputTexture = GraphBuilder.CreateTexture(FilteredRadianceOutputDesc, TEXT("NFOR.FilteredRadiance.Output"));
+				AddCopyMirroredTexturePass(GraphBuilder, FilteredRadiance, FilteredRadianceOutputTexture);
+				FilteredRadiance = FilteredRadianceOutputTexture;
+			}
 		}
 
 		return FilteredRadiance;
