Fallback to more stable linear solver in NFOR until the robustness of the fast solver is improved. The time for denoising single 1k image increased from 0.8s to 1.2s.

#tests No artifacts in the hillside internal scene

diff --git a/Engine/Plugins/Experimental/NFORDenoise/Source/NFORDenoise/Private/NFORDenoiseCS.cpp b/Engine/Plugins/Experimental/NFORDenoise/Source/NFORDenoise/Private/NFORDenoiseCS.cpp
--- a/Engine/Plugins/Experimental/NFORDenoise/Source/NFORDenoise/Private/NFORDenoiseCS.cpp
+++ b/Engine/Plugins/Experimental/NFORDenoise/Source/NFORDenoise/Private/NFORDenoiseCS.cpp
@@ -134,7 +134,7 @@ namespace NFORDenoise
 
 	TAutoConsoleVariable<int32> CVarNFORLinearSolverType(
 		TEXT("r.NFOR.LinearSolver.Type"),
-		2,
+		0,
 		TEXT("The linear regression solver type implemented in GPU.\n")
 		TEXT("0: Newton Schulz iterative method (High quality but slow).\n")
 		TEXT("1: Cholesky decomposition (Fast but has too smoothed result or artifacts).\n")
