Fix normal not scaled by vector length in NFOR denoiser.

#rnx

diff --git a/Engine/Plugins/Experimental/NFORDenoise/Source/NFORDenoise/Private/NFORDenoise.cpp b/Engine/Plugins/Experimental/NFORDenoise/Source/NFORDenoise/Private/NFORDenoise.cpp
--- a/Engine/Plugins/Experimental/NFORDenoise/Source/NFORDenoise/Private/NFORDenoise.cpp
+++ b/Engine/Plugins/Experimental/NFORDenoise/Source/NFORDenoise/Private/NFORDenoise.cpp
@@ -81,7 +81,7 @@ void FRDGFrameCache::AppendToNFORBuffer(
 		NFORDenoise::FNFORTextureDesc NormalVarianceTex(RadianceVariance, 2, 1, 4);
 
 		Features.Add(NFORDenoise::FFeatureDesc(AlbedoTex, AlbedoVarianceTex, NFORDenoise::EVarianceType::GreyScale));
-		Features.Add(NFORDenoise::FFeatureDesc(NormalTex, NormalVarianceTex, NFORDenoise::EVarianceType::GreyScale));
+		Features.Add(NFORDenoise::FFeatureDesc(NormalTex, NormalVarianceTex, NFORDenoise::EVarianceType::Normal));
 
 		if (IsFeatureDepthEnabled())
 		{
