[NFOR] Remove one unnecessary additional base memory allocation for radiance and feature caching.

#tests Test shows it saves between 100~200 MB when using nvidia-smi to print out the memory usage for 5 frame history.
#Synced-CL 49128105

diff --git a/Engine/Plugins/Experimental/NFORDenoise/Source/NFORDenoise/Private/NFORDenoise.cpp b/Engine/Plugins/Experimental/NFORDenoise/Source/NFORDenoise/Private/NFORDenoise.cpp
--- a/Engine/Plugins/Experimental/NFORDenoise/Source/NFORDenoise/Private/NFORDenoise.cpp
+++ b/Engine/Plugins/Experimental/NFORDenoise/Source/NFORDenoise/Private/NFORDenoise.cpp
@@ -131,11 +131,29 @@ public:
 			RDG_EVENT_SCOPE(GraphBuilder, "AddFrame");
 
 			FRDGFrameCache RDGFrameCache;
-			RDGFrameCache.Radiance = GraphBuilder.CreateTexture(Inputs.ColorTex->Desc,TEXT("NFOR.Radiance"), ERDGTextureFlags::MultiFrame);
-			RDGFrameCache.RadianceVariance = GraphBuilder.CreateTexture(Inputs.VarianceTex->Desc, TEXT("NFOR.RadianceVariance"), ERDGTextureFlags::MultiFrame);
-			RDGFrameCache.Albedo = GraphBuilder.CreateTexture(Inputs.AlbedoTex->Desc, TEXT("NFOR.Albedo"), ERDGTextureFlags::MultiFrame);
-			RDGFrameCache.Normal = GraphBuilder.CreateTexture(Inputs.NormalTex->Desc, TEXT("NFOR.Normal"), ERDGTextureFlags::MultiFrame);
-			RDGFrameCache.Depth = GraphBuilder.CreateTexture(Inputs.DepthTex->Desc, TEXT("NFOR.Depth"), ERDGTextureFlags::MultiFrame);
+			
+			if (FrameCaches.Num() < NumFrames)
+			{
+				RDGFrameCache.Radiance = GraphBuilder.CreateTexture(Inputs.ColorTex->Desc, TEXT("NFOR.Radiance"), ERDGTextureFlags::MultiFrame);
+				RDGFrameCache.RadianceVariance = GraphBuilder.CreateTexture(Inputs.VarianceTex->Desc, TEXT("NFOR.RadianceVariance"), ERDGTextureFlags::MultiFrame);
+				RDGFrameCache.Albedo = GraphBuilder.CreateTexture(Inputs.AlbedoTex->Desc, TEXT("NFOR.Albedo"), ERDGTextureFlags::MultiFrame);
+				RDGFrameCache.Normal = GraphBuilder.CreateTexture(Inputs.NormalTex->Desc, TEXT("NFOR.Normal"), ERDGTextureFlags::MultiFrame);
+				RDGFrameCache.Depth = GraphBuilder.CreateTexture(Inputs.DepthTex->Desc, TEXT("NFOR.Depth"), ERDGTextureFlags::MultiFrame);
+			}
+			else
+			{
+				// Remove unreachable caches if Frame count is changed between each denoising cycle.
+				while (FrameCaches.Num() > NumFrames)
+				{
+					FrameCaches.RemoveAt(0);
+				}
+
+				// Avoid potential reallocation.
+				auto First = MoveTemp(FrameCaches[0]);
+				FrameCaches.RemoveAt(0, 1, EAllowShrinking::No);
+				FrameCaches.Add(MoveTemp(First));
+				RDGFrameCache = FrameCaches[NumFrames - 1].ToRGDResource(GraphBuilder, View);
+			}
 
 			AddCopyTexturePass(GraphBuilder, Inputs.ColorTex, RDGFrameCache.Radiance);
 			AddCopyTexturePass(GraphBuilder, Inputs.VarianceTex, RDGFrameCache.RadianceVariance);
@@ -143,13 +161,11 @@ public:
 			AddCopyTexturePass(GraphBuilder, Inputs.NormalTex, RDGFrameCache.Normal);
 			AddCopyTexturePass(GraphBuilder, Inputs.DepthTex, RDGFrameCache.Depth);
 			
-			FExternalFrameCache FrameCache = RDGFrameCache.ToExternalResource(GraphBuilder, View);
-
-			while (FrameCaches.Num() >= NumFrames)
+			if (FrameCaches.Num() < NumFrames)
 			{
-				FrameCaches.RemoveAt(0);
+				FExternalFrameCache FrameCache = RDGFrameCache.ToExternalResource(GraphBuilder, View);
+				FrameCaches.Add(MoveTemp(FrameCache));
 			}
-			FrameCaches.Add(MoveTemp(FrameCache));
 		}
 
 		void FetchFrames(FRDGBuilder& GraphBuilder, const FSceneView& View, 
