[NFOR] Add low memory mode for memory constrained GPUs. All 32bit channel cache will be stored in 16bit when `r.NFOR.Memory.Mode = 0`. Might have small quality drop. 4k rendering total cached texture memory usage reduced from 2814MB to 2113 MB when Frame Count = 2.

#tests local projects produce similar results when compared to high memory mode.

#Synced-CL 49788808

diff --git a/Engine/Plugins/Experimental/NFORDenoise/Source/NFORDenoise/Private/NFORDenoise.cpp b/Engine/Plugins/Experimental/NFORDenoise/Source/NFORDenoise/Private/NFORDenoise.cpp
--- a/Engine/Plugins/Experimental/NFORDenoise/Source/NFORDenoise/Private/NFORDenoise.cpp
+++ b/Engine/Plugins/Experimental/NFORDenoise/Source/NFORDenoise/Private/NFORDenoise.cpp
@@ -21,6 +21,38 @@ static TAutoConsoleVariable<bool> CVarNFORFeatureDepth(
 	ECVF_RenderThreadSafe
 );
 
+static TAutoConsoleVariable<int32> CVarNFORMemoryMode(
+	TEXT("r.NFOR.Memory.Mode"),
+	1,
+	TEXT("The memory mode to use when caching textures.\n")
+	TEXT("0: Low memory mode. All 32bit radiance will be cached in 16bit.\n")
+	TEXT("1: High memory mode. Use the default texture format used in path tracing."),
+	ECVF_RenderThreadSafe
+);
+
+bool UseLowMemoryMode()
+{
+	return CVarNFORMemoryMode.GetValueOnRenderThread() == 0;
+}
+
+FRDGTextureDesc ApplyMemoryMode(const FRDGTextureDesc& Desc)
+{
+
+	if (!UseLowMemoryMode())
+	{
+		return Desc;
+	}
+
+	FRDGTextureDesc TmpDesc = Desc;
+
+	if (Desc.Format == PF_A32B32G32R32F)
+	{
+		TmpDesc.Format = PF_FloatRGBA;
+	}
+
+	return TmpDesc;
+}
+
 bool IsFeatureDepthEnabled()
 {
 	return CVarNFORFeatureDepth.GetValueOnRenderThread();
@@ -134,8 +166,11 @@ public:
 			
 			if (FrameCaches.Num() < NumFrames)
 			{
-				RDGFrameCache.Radiance = GraphBuilder.CreateTexture(Inputs.ColorTex->Desc, TEXT("NFOR.Radiance"), ERDGTextureFlags::MultiFrame);
-				RDGFrameCache.RadianceVariance = GraphBuilder.CreateTexture(Inputs.VarianceTex->Desc, TEXT("NFOR.RadianceVariance"), ERDGTextureFlags::MultiFrame);
+				RDGFrameCache.Radiance = GraphBuilder.CreateTexture(
+					ApplyMemoryMode(Inputs.ColorTex->Desc), TEXT("NFOR.Radiance"), ERDGTextureFlags::MultiFrame);
+				RDGFrameCache.RadianceVariance = GraphBuilder.CreateTexture(
+					ApplyMemoryMode(Inputs.VarianceTex->Desc), TEXT("NFOR.RadianceVariance"), ERDGTextureFlags::MultiFrame);
+
 				RDGFrameCache.Albedo = GraphBuilder.CreateTexture(Inputs.AlbedoTex->Desc, TEXT("NFOR.Albedo"), ERDGTextureFlags::MultiFrame);
 				RDGFrameCache.Normal = GraphBuilder.CreateTexture(Inputs.NormalTex->Desc, TEXT("NFOR.Normal"), ERDGTextureFlags::MultiFrame);
 				RDGFrameCache.Depth = GraphBuilder.CreateTexture(Inputs.DepthTex->Desc, TEXT("NFOR.Depth"), ERDGTextureFlags::MultiFrame);
@@ -155,8 +190,18 @@ public:
 				RDGFrameCache = FrameCaches[NumFrames - 1].ToRGDResource(GraphBuilder, View);
 			}
 
-			AddCopyTexturePass(GraphBuilder, Inputs.ColorTex, RDGFrameCache.Radiance);
-			AddCopyTexturePass(GraphBuilder, Inputs.VarianceTex, RDGFrameCache.RadianceVariance);
+			if (UseLowMemoryMode())
+			{
+				// Radiance and radiance variance source and target can have different 4 channel format. Use customized function.
+				NFORDenoise::AddCopyMirroredTexturePass(GraphBuilder, Inputs.ColorTex, RDGFrameCache.Radiance);
+				NFORDenoise::AddCopyMirroredTexturePass(GraphBuilder, Inputs.VarianceTex, RDGFrameCache.RadianceVariance);
+			}
+			else
+			{
+				AddCopyTexturePass(GraphBuilder, Inputs.ColorTex, RDGFrameCache.Radiance);
+				AddCopyTexturePass(GraphBuilder, Inputs.VarianceTex, RDGFrameCache.RadianceVariance);
+			}
+
 			AddCopyTexturePass(GraphBuilder, Inputs.AlbedoTex, RDGFrameCache.Albedo);
 			AddCopyTexturePass(GraphBuilder, Inputs.NormalTex, RDGFrameCache.Normal);
 			AddCopyTexturePass(GraphBuilder, Inputs.DepthTex, RDGFrameCache.Depth);
diff --git a/Engine/Plugins/Experimental/NFORDenoise/Source/NFORDenoise/Private/NFORDenoiseCS.cpp b/Engine/Plugins/Experimental/NFORDenoise/Source/NFORDenoise/Private/NFORDenoiseCS.cpp
--- a/Engine/Plugins/Experimental/NFORDenoise/Source/NFORDenoise/Private/NFORDenoiseCS.cpp
+++ b/Engine/Plugins/Experimental/NFORDenoise/Source/NFORDenoise/Private/NFORDenoiseCS.cpp
@@ -2864,7 +2864,8 @@ namespace NFORDenoise
 		}
 		else
 		{
-			AddCopyTexturePass(GraphBuilder, FilteredImages[0], DenoisedRadiance);
+			// Support different 4 channel format texture copy
+			AddCopyMirroredTexturePass(GraphBuilder, FilteredImages[0], DenoisedRadiance);
 		}
 
 		{
