Disabled compiling of NFOR shaders on Metal platforms due to atomicBinOp.i64

diff --git a/Engine/Plugins/Experimental/NFORDenoise/Source/NFORDenoise/Private/NFORDenoiseCS.cpp b/Engine/Plugins/Experimental/NFORDenoise/Source/NFORDenoise/Private/NFORDenoiseCS.cpp
--- a/Engine/Plugins/Experimental/NFORDenoise/Source/NFORDenoise/Private/NFORDenoiseCS.cpp
+++ b/Engine/Plugins/Experimental/NFORDenoise/Source/NFORDenoise/Private/NFORDenoiseCS.cpp
@@ -279,7 +279,8 @@ namespace NFORDenoise
 
 		return ShouldCompileRayTracingShadersForProject(ShaderPlatform) &&
 			FDataDrivenShaderPlatformInfo::GetSupportsPathTracing(ShaderPlatform) &&
-			bSupportsPathTracing;
+			bSupportsPathTracing &&
+			!IsMetalPlatform(ShaderPlatform);
 	}
 
 	bool ShouldFeatureAddConstant()
